UAFPoseSearch - fix for UUAFGraphNodeTemplate_MotionMatching traits order

diff --git a/Engine/Plugins/Experimental/UAF/UAFPoseSearch/Source/UAFPoseSearchUncookedOnly/Private/UAFGraphNodeTemplate_MotionMatching.h b/Engine/Plugins/Experimental/UAF/UAFPoseSearch/Source/UAFPoseSearchUncookedOnly/Private/UAFGraphNodeTemplate_MotionMatching.h
--- a/Engine/Plugins/Experimental/UAF/UAFPoseSearch/Source/UAFPoseSearchUncookedOnly/Private/UAFGraphNodeTemplate_MotionMatching.h
+++ b/Engine/Plugins/Experimental/UAF/UAFPoseSearch/Source/UAFPoseSearchUncookedOnly/Private/UAFGraphNodeTemplate_MotionMatching.h
@@ -1,4 +1,4 @@
-﻿// Copyright Epic Games, Inc. All Rights Reserved.
+// Copyright Epic Games, Inc. All Rights Reserved.
 
 #pragma once
 
@@ -31,8 +31,8 @@ class UUAFGraphNodeTemplate_MotionMatching : public UUAFGraphNodeTemplate
 		Traits =
 		{
 			TInstancedStruct<FAnimNextBlendStackCoreTraitSharedData>::Make(),
-			TInstancedStruct<FMotionMatchingTraitSharedData>::Make(),
 			TInstancedStruct<FAnimNextBlendSmootherCoreTraitSharedData>::Make(),
+			TInstancedStruct<FMotionMatchingTraitSharedData>::Make(),
 			TInstancedStruct<FAnimNextHistoryCollectorTraitSharedData>::Make()
 		};
 		DragDropAssetTypes.Add(UPoseSearchDatabase::StaticClass());
