UAFPoseSearch - removing unneeded duplicate line of code in FMotionMatchingTrait::AddReferencedObjects

diff --git a/Engine/Plugins/Experimental/UAF/UAFPoseSearch/Source/UAFPoseSearch/Private/MotionMatchingTrait.cpp b/Engine/Plugins/Experimental/UAF/UAFPoseSearch/Source/UAFPoseSearch/Private/MotionMatchingTrait.cpp
--- a/Engine/Plugins/Experimental/UAF/UAFPoseSearch/Source/UAFPoseSearch/Private/MotionMatchingTrait.cpp
+++ b/Engine/Plugins/Experimental/UAF/UAFPoseSearch/Source/UAFPoseSearch/Private/MotionMatchingTrait.cpp
@@ -658,7 +658,6 @@ PRAGMA_ENABLE_DEPRECATION_WARNINGS
 
 		Collector.AddReferencedObject(SearchResult.SelectedAnim);
 		Collector.AddReferencedObject(SearchResult.SelectedDatabase);
-		Collector.AddReferencedObject(SearchResult.SelectedDatabase);
 
 		for (TObjectPtr<const UObject>& AnimContexts : SearchResult.AnimContexts)
 		{
