PoseSearch - fix for MM trait timeline and MM node interaction blendstack synchronizations

#[review]-45804074 [REDACTED_HANDLE]

diff --git a/Engine/Plugins/Animation/PoseSearch/Source/Runtime/Private/AnimNode_MotionMatching.cpp b/Engine/Plugins/Animation/PoseSearch/Source/Runtime/Private/AnimNode_MotionMatching.cpp
--- a/Engine/Plugins/Animation/PoseSearch/Source/Runtime/Private/AnimNode_MotionMatching.cpp
+++ b/Engine/Plugins/Animation/PoseSearch/Source/Runtime/Private/AnimNode_MotionMatching.cpp
@@ -104,16 +104,6 @@ void FAnimNode_MotionMatching::UpdateAssetPlayer(const FAnimationUpdateContext&
 		MotionMatchingState.Reset();
 	}
 #endif // WITH_EDITOR
-	else
-	{
-		// We adjust the motion matching state asset time to the current player node's asset time. This is done 
-		// because the player node may have ticked more or less time than we expected due to variable dt or the 
-		// dynamic playback rate adjustment and as such the motion matching state does not update by itself
-		MotionMatchingState.SearchResult.SelectedAnim = GetAnimAsset();
-		MotionMatchingState.SearchResult.SelectedTime = GetAccumulatedTime();
-		MotionMatchingState.SearchResult.bIsMirrored = GetMirror();
-		MotionMatchingState.SearchResult.BlendParameters = GetBlendParameters();
-	}
 
 	UpdateCounter.SynchronizeWith(Context.AnimInstanceProxy->GetUpdateCounter());
 
@@ -168,11 +158,20 @@ void FAnimNode_MotionMatching::UpdateAssetPlayer(const FAnimationUpdateContext&
 
 		const UObject* AnimContext = Context.AnimInstanceProxy->GetAnimInstanceObject();
 		check(AnimContext);
+		// MotionMatchInteraction synchronizes and integrates MotionMatchingState.SearchResult internally, so there's no need to update it with the blendstack information
 		const bool bIsInteraction = UPoseSearchInteractionLibrary::MotionMatchInteraction(MotionMatchingState.SearchResult, Availabilities, AnimContext, FName(), PoseHistory, bValidateResultAgainstAvailabilities, bKeepInteractionAlive, BlendTime);
 
 		// performing the regular single character motion matching search in case there's no MM interaction
 		if (!bIsInteraction)
 		{
+			// We adjust the motion matching state asset time to the current player node's asset time. This is done 
+			// because the player node may have ticked more or less time than we expected due to variable dt or the 
+			// dynamic playback rate adjustment and as such the motion matching state does not update by itself
+			MotionMatchingState.SearchResult.SelectedAnim = GetAnimAsset();
+			MotionMatchingState.SearchResult.SelectedTime = GetAccumulatedTime();
+			MotionMatchingState.SearchResult.bIsMirrored = GetMirror();
+			MotionMatchingState.SearchResult.BlendParameters = GetBlendParameters();
+
 			FFloatInterval PoseSearchPlayRate = PlayRate;
 #if ENABLE_ANIM_DEBUG
 			bool bPoseSearchPlayRateEnabled = GVarAnimNodeMotionMatchingPlayRateEnabled == UE::Private::EPlayRateState::Enabled || GVarAnimNodeMotionMatchingPlayRateEnabled == UE::Private::EPlayRateState::PoseSearchOnly;
diff --git a/Engine/Plugins/Experimental/UAF/UAFPoseSearch/Source/UAFPoseSearch/Private/MotionMatchingTrait.cpp b/Engine/Plugins/Experimental/UAF/UAFPoseSearch/Source/UAFPoseSearch/Private/MotionMatchingTrait.cpp
--- a/Engine/Plugins/Experimental/UAF/UAFPoseSearch/Source/UAFPoseSearch/Private/MotionMatchingTrait.cpp
+++ b/Engine/Plugins/Experimental/UAF/UAFPoseSearch/Source/UAFPoseSearch/Private/MotionMatchingTrait.cpp
@@ -157,44 +157,56 @@ PRAGMA_ENABLE_DEPRECATION_WARNINGS
 		}
 #endif // WITH_EDITOR
 
-		// We adjust the MotionMatchingInteractionState time to the current player node's asset time. This is done 
-		// because the player node may have ticked more or less time than we expected due to variable dt or the 
-		// dynamic playback rate adjustment and as such the MotionMatchingInteractionState does not update by itself
-		const FTimelineState TimelineState = TimelineTrait.GetState(Context);
-		const float TimelineRealTime = TimelineState.GetPosition();
-		float TimelineNormalizedTime = TimelineRealTime;
-
-		// @todo: This is a hack, since we don't have a way to get the normalized timeline of a blendspace yet (pending syncing...?)
-		if (Cast<UBlendSpace>(MotionMatchingState.SearchResult.GetAnimationAssetForRole()))
-		{
-			// Convert to normalized time.
-			const float TimelineDuration = TimelineState.GetDuration();
-			if (TimelineDuration > 0.0f)
-			{
-				// NOTE: This doesn't work because AnimNext normalized time does not match database indexing normalized time.
-				TimelineNormalizedTime = TimelineRealTime / TimelineDuration;
-			}
-		}
 
-#if ENABLE_VISUAL_LOG
-		if (FVisualLogger::IsRecording())
-		{
-			const FName AnimName = TimelineState.GetDebugName();
-			const float Duration = TimelineState.GetDuration();
-			static const TCHAR* LogName = TEXT("FMotionMatchingTrait");
-			UE_VLOG(Context.GetHostObject(), LogName, Verbose, TEXT("TimelineRealTime: %f"), TimelineRealTime);
-			UE_VLOG(Context.GetHostObject(), LogName, Verbose, TEXT("TimelineNormalizedTime: %f"), TimelineNormalizedTime);
-			UE_VLOG(Context.GetHostObject(), LogName, Verbose, TEXT("AnimName: %s"), *AnimName.ToString());
-			UE_VLOG(Context.GetHostObject(), LogName, Verbose, TEXT("Duration: %f"), Duration);
-		}
-#endif
 
 		const TArrayView<const FPoseSearchInteractionAvailability> Availabilities = SharedData->GetAvailabilities(Binding);
+		// MotionMatchInteraction synchronizes and integrates MotionMatchingState.SearchResult internally, so there's no need to update it with the blendstack information
 		const bool bIsInteraction = UPoseSearchInteractionLibrary::MotionMatchInteraction(MotionMatchingState.SearchResult, Availabilities, AnimContext, FName(), PoseHistory, SharedData->GetbValidateResultAgainstAvailabilities(Binding), SharedData->GetbKeepInteractionAlive(Binding), SharedData->BlendArguments.BlendTime);
 
 		// performing the regular single character motion matching search in case there's no MM interaction
 		if (!bIsInteraction)
 		{
+			// We adjust the MotionMatchingInteractionState time to the current player node's asset time. This is done 
+			// because the player node may have ticked more or less time than we expected due to variable dt or the 
+			// dynamic playback rate adjustment and as such the MotionMatchingInteractionState does not update by itself
+			const FTimelineState TimelineState = TimelineTrait.GetState(Context);
+			const float TimelineRealTime = TimelineState.GetPosition();
+			float TimelineNormalizedTime = TimelineRealTime;
+
+			// @todo: This is a hack, since we don't have a way to get the normalized timeline of a blendspace yet (pending syncing...?)
+			if (Cast<UBlendSpace>(MotionMatchingState.SearchResult.GetAnimationAssetForRole()))
+			{
+				// Convert to normalized time.
+				const float TimelineDuration = TimelineState.GetDuration();
+				if (TimelineDuration > 0.0f)
+				{
+					// NOTE: This doesn't work because AnimNext normalized time does not match database indexing normalized time.
+					TimelineNormalizedTime = TimelineRealTime / TimelineDuration;
+				}
+			}
+
+#if ENABLE_VISUAL_LOG
+			if (FVisualLogger::IsRecording())
+			{
+				const FName AnimName = TimelineState.GetDebugName();
+				const float Duration = TimelineState.GetDuration();
+				static const TCHAR* LogName = TEXT("FMotionMatchingTrait");
+				UE_VLOG(Context.GetHostObject(), LogName, Verbose, TEXT("TimelineRealTime: %f"), TimelineRealTime);
+				UE_VLOG(Context.GetHostObject(), LogName, Verbose, TEXT("TimelineNormalizedTime: %f"), TimelineNormalizedTime);
+				UE_VLOG(Context.GetHostObject(), LogName, Verbose, TEXT("AnimName: %s"), *AnimName.ToString());
+				UE_VLOG(Context.GetHostObject(), LogName, Verbose, TEXT("Duration: %f"), Duration);
+			}
+#endif
+
+			// @todo: ask blendstack what it's playing instead of relying on the previous MotionMatchingState.SearchResult
+			// We adjust the motion matching state asset time to the current player node's asset time. This is done 
+			// because the player node may have ticked more or less time than we expected due to variable dt or the 
+			// dynamic playback rate adjustment and as such the motion matching state does not update by itself
+			//MotionMatchingState.SearchResult.SelectedAnim = GetAnimAsset();
+			MotionMatchingState.SearchResult.SelectedTime = TimelineNormalizedTime;
+			//MotionMatchingState.SearchResult.bIsMirrored = GetMirror();
+			//MotionMatchingState.SearchResult.BlendParameters = GetBlendParameters();
+
 			const TConstArrayView<TObjectPtr<const UPoseSearchDatabase>> DatabasesToSearch = SharedData->GetDatabases(Binding);
 			if (DatabasesToSearch.IsEmpty())
 			{
