Added UAF GenerateTrajectory version for CharacterMovementComponent.

diff --git a/Engine/Plugins/Experimental/UAF/UAFPoseSearch/Source/UAFPoseSearch/Private/RigUnit_GenerateCharacterMovementComponentTrajectory.cpp b/Engine/Plugins/Experimental/UAF/UAFPoseSearch/Source/UAFPoseSearch/Private/RigUnit_GenerateCharacterMovementComponentTrajectory.cpp
new file mode 100644
--- /dev/null
+++ b/Engine/Plugins/Experimental/UAF/UAFPoseSearch/Source/UAFPoseSearch/Private/RigUnit_GenerateCharacterMovementComponentTrajectory.cpp
@@ -0,0 +1,38 @@
+﻿// Copyright Epic Games, Inc. All Rights Reserved.
+
+#include "RigUnit_GenerateCharacterMovementComponentTrajectory.h"
+
+#include "GameFramework/CharacterMovementComponent.h"
+#include "PoseSearch/PoseSearchDefines.h"
+#include "PoseSearch/PoseSearchTrajectoryLibrary.h"
+
+FRigUnit_GenerateCharacterMovementComponentTrajectory_Execute()
+{
+	TRACE_CPUPROFILER_EVENT_SCOPE(FRigUnit_GenerateCharacterMovementComponentTrajectory_Execute);
+	using namespace UE::UAF;
+
+	if (CharacterMovementComponent == nullptr)
+	{
+		return;
+	}
+
+	AActor* Owner = CharacterMovementComponent->GetOwner();
+	if (Owner == nullptr)
+	{
+		UE_LOG(LogPoseSearch, Warning, TEXT("FRigUnit_GenerateCharacterMovementComponentTrajectory_Execute: Failed to get Character Movement Component's actor owner."));
+		return;
+	}
+	
+	UPoseSearchTrajectoryLibrary::PoseSearchGenerateTransformTrajectory(
+		Owner,
+		TrajectoryData,
+		DeltaTime,
+		InOutTrajectory,
+		InOutDesiredControllerYawLastUpdate,
+		InOutTrajectory,
+		HistorySamplingInterval,
+		NumHistorySamples,
+		PredictionSamplingInterval,
+		NumPredictionSamples
+	);
+}
diff --git a/Engine/Plugins/Experimental/UAF/UAFPoseSearch/Source/UAFPoseSearch/Private/RigUnit_GenerateCharacterMovementComponentTrajectory.h b/Engine/Plugins/Experimental/UAF/UAFPoseSearch/Source/UAFPoseSearch/Private/RigUnit_GenerateCharacterMovementComponentTrajectory.h
new file mode 100644
--- /dev/null
+++ b/Engine/Plugins/Experimental/UAF/UAFPoseSearch/Source/UAFPoseSearch/Private/RigUnit_GenerateCharacterMovementComponentTrajectory.h
@@ -0,0 +1,52 @@
+﻿// Copyright Epic Games, Inc. All Rights Reserved.
+
+#pragma once
+
+#include "AnimNextExecuteContext.h"
+#include "Graph/RigUnit_AnimNextBase.h"
+#include "Animation/TrajectoryTypes.h"
+#include "GameFramework/CharacterMovementComponent.h"
+#include "PoseSearch/PoseSearchTrajectoryLibrary.h"
+
+#include "RigUnit_GenerateCharacterMovementComponentTrajectory.generated.h"
+
+USTRUCT(meta = (DisplayName = "Generate Trajectory from Character Movement Component", Category = "Character Movement", NodeColor = "0, 1, 1", Keywords = "Trajectory"))
+struct FRigUnit_GenerateCharacterMovementComponentTrajectory : public FRigUnit_AnimNextBase
+{
+	GENERATED_BODY()
+
+	RIGVM_METHOD()
+	void Execute();
+
+	UPROPERTY(EditAnywhere, Category = "Trajectory", meta = (Input))
+	TObjectPtr<UCharacterMovementComponent> CharacterMovementComponent;
+
+	UPROPERTY(EditAnywhere, Category = "Trajectory", meta = (Input))
+	// This should be the most recent simulation time that was used to get us to our current state
+	float DeltaTime = 0.f;
+
+	UPROPERTY(EditAnywhere, Category = "Trajectory", meta = (Input))
+	float HistorySamplingInterval = -1.f;
+
+	UPROPERTY(EditAnywhere, Category = "Trajectory", meta = (Input))
+	int32 NumHistorySamples = 30;
+
+	UPROPERTY(EditAnywhere, Category = "Trajectory", meta = (Input))
+	float PredictionSamplingInterval = 0.1f;
+
+	UPROPERTY(EditAnywhere, Category = "Trajectory", meta = (Input))
+	int32 NumPredictionSamples = 15;
+
+	UPROPERTY(EditAnywhere, Category = "Trajectory", meta = (Input))
+	FPoseSearchTrajectoryData TrajectoryData;
+	
+	UPROPERTY(EditAnywhere, Category = "Trajectory", meta = (Input, Output))
+	FTransformTrajectory InOutTrajectory;
+
+	UPROPERTY(EditAnywhere, Category = "Trajectory", meta = (Input, Output))
+	float InOutDesiredControllerYawLastUpdate = 0.0f;
+	
+	// The execution result
+	UPROPERTY(EditAnywhere, DisplayName = "Execute", Category = "BeginExecution", meta = (Input, Output))
+	FAnimNextExecuteContext ExecuteContext;
+};
