UAF: Fixed LOD out of bounds in pose history trait.

diff --git a/Engine/Plugins/Experimental/UAF/UAFPoseSearch/Source/UAFPoseSearch/Private/HistoryCollectorTrait.cpp b/Engine/Plugins/Experimental/UAF/UAFPoseSearch/Source/UAFPoseSearch/Private/HistoryCollectorTrait.cpp
--- a/Engine/Plugins/Experimental/UAF/UAFPoseSearch/Source/UAFPoseSearch/Private/HistoryCollectorTrait.cpp
+++ b/Engine/Plugins/Experimental/UAF/UAFPoseSearch/Source/UAFPoseSearch/Private/HistoryCollectorTrait.cpp
@@ -41,7 +41,7 @@ namespace UE::UAF
 			const TArrayView<const FBoneIndexType> SkeletonBoneIndexToLODBoneIndexMap = Pose.GetSkeletonBoneIndexToLODBoneIndexMap();
 			const FBoneIndexType LODBoneIndex = SkeletonBoneIndexToLODBoneIndexMap[SkeletonBoneIdx.GetInt()];
 			
-			if (LODBoneIndex != FBoneIndexType(INDEX_NONE))
+			if (LODBoneIndex != FBoneIndexType(INDEX_NONE) && LODBoneIndex < Pose.GetNumBones())
 			{
 				ComponentSpaceTransform = Pose.LocalTransformsView[LODBoneIndex];
 			}
