UAFPoseSearch - initial MM interaction warping support

#Synced-CL 48365056

diff --git a/Engine/Plugins/Animation/PoseSearch/Source/Runtime/Private/AnimNode_MotionMatching.cpp b/Engine/Plugins/Animation/PoseSearch/Source/Runtime/Private/AnimNode_MotionMatching.cpp
--- a/Engine/Plugins/Animation/PoseSearch/Source/Runtime/Private/AnimNode_MotionMatching.cpp
+++ b/Engine/Plugins/Animation/PoseSearch/Source/Runtime/Private/AnimNode_MotionMatching.cpp
@@ -313,6 +313,11 @@ void FAnimNode_MotionMatching::Evaluate_AnyThread(FPoseContext& Output)
 
 	Super::Evaluate_AnyThread(Output);
 
+	if (!bEnableWarp)
+	{
+		return;
+	}
+
 	// calculating the animation time we want the full alignment transforms from
 	UAnimInstance* AnimInstance = Cast<UAnimInstance>(Output.AnimInstanceProxy->GetAnimInstanceObject());
 	if (!AnimInstance)
diff --git a/Engine/Plugins/Animation/PoseSearch/Source/Runtime/Public/PoseSearch/AnimNode_MotionMatching.h b/Engine/Plugins/Animation/PoseSearch/Source/Runtime/Public/PoseSearch/AnimNode_MotionMatching.h
--- a/Engine/Plugins/Animation/PoseSearch/Source/Runtime/Public/PoseSearch/AnimNode_MotionMatching.h
+++ b/Engine/Plugins/Animation/PoseSearch/Source/Runtime/Public/PoseSearch/AnimNode_MotionMatching.h
@@ -192,6 +192,13 @@ private:
 	UPROPERTY(EditAnywhere, Category = "Interaction | Experimental", meta = (PinHiddenByDefault))
 	bool bKeepInteractionAlive = false;
 
+	// Experimental, this feature might be removed without warning, not for production use
+	// if bEnableWarp is true and this node is performing an interaction, warp aligment will be caluculated
+	// and requested to be applied by the anim root motion provider (IAnimRootMotionProvider interface)
+	// set it to false if a different warp aligment strategy is used
+	UPROPERTY(EditAnywhere, Category = "Interaction | Experimental", meta = (PinHiddenByDefault))
+	bool bEnableWarp = false;
+
 	// Experimental, this feature might be removed without warning, not for production use
 	// if bWarpUsingRootBone is true, warping will be calculated using the interacting actors previous frame root bone transforms (effective for setups with OffsetRootBone node allowing root bone drifting from capsule)
 	// if bWarpUsingRootBone is true, warping will be calculated using the previous frame root transforms (effective root motion driven for setups)
diff --git a/Engine/Plugins/Experimental/UAF/UAFPoseSearch/Source/UAFPoseSearch/Internal/MotionMatchingTraitData.h b/Engine/Plugins/Experimental/UAF/UAFPoseSearch/Source/UAFPoseSearch/Internal/MotionMatchingTraitData.h
--- a/Engine/Plugins/Experimental/UAF/UAFPoseSearch/Source/UAFPoseSearch/Internal/MotionMatchingTraitData.h
+++ b/Engine/Plugins/Experimental/UAF/UAFPoseSearch/Source/UAFPoseSearch/Internal/MotionMatchingTraitData.h
@@ -108,6 +108,13 @@ struct FMotionMatchingTraitSharedData : public FAnimNextTraitSharedData
 	UPROPERTY(EditAnywhere, Category = Interaction, meta = (PinHiddenByDefault))
 	bool bKeepInteractionAlive = false;
 
+	// Experimental, this feature might be removed without warning, not for production use
+	// if bEnableWarp is true and this node is performing an interaction, warp aligment will be caluculated
+	// and requested to be applied by the anim root motion provider (IAnimRootMotionProvider interface)
+	// set it to false if a different warp aligment strategy is used
+	UPROPERTY(EditAnywhere, Category = Interaction, meta = (PinHiddenByDefault))
+	bool bEnableWarp = false;
+
 	UPROPERTY(EditAnywhere, Category = Interaction, meta = (PinHiddenByDefault))
 	bool bWarpUsingRootBone = true;
 
@@ -149,6 +156,7 @@ struct FMotionMatchingTraitSharedData : public FAnimNextTraitSharedData
 		GeneratorMacro(Availabilities) \
 		GeneratorMacro(bValidateResultAgainstAvailabilities) \
 		GeneratorMacro(bKeepInteractionAlive) \
+		GeneratorMacro(bEnableWarp) \
 		GeneratorMacro(bWarpUsingRootBone) \
 		GeneratorMacro(WarpingRotationRatio) \
 		GeneratorMacro(WarpingTranslationRatio) \
diff --git a/Engine/Plugins/Experimental/UAF/UAFPoseSearch/Source/UAFPoseSearch/Private/Module.cpp b/Engine/Plugins/Experimental/UAF/UAFPoseSearch/Source/UAFPoseSearch/Private/Module.cpp
--- a/Engine/Plugins/Experimental/UAF/UAFPoseSearch/Source/UAFPoseSearch/Private/Module.cpp
+++ b/Engine/Plugins/Experimental/UAF/UAFPoseSearch/Source/UAFPoseSearch/Private/Module.cpp
@@ -1,19 +1,19 @@
 // Copyright Epic Games, Inc. All Rights Reserved.
 
+#include "AlphaBlend.h"
+#include "Animation/TrajectoryTypes.h"
+#include "AnimNextInteractionIslandDependency.h"
 #include "CoreMinimal.h"
+#include "Features/IModularFeatures.h"
 #include "Module/AnimNextModule.h"
 #include "Modules/ModuleInterface.h"
 #include "Modules/ModuleManager.h"
-#include "RigVMCore/RigVMRegistry.h"
-
-#include "AlphaBlend.h"
-#include "AnimNextInteractionIslandDependency.h"
-#include "Features/IModularFeatures.h"
 #include "PoseSearch/PoseSearchDatabase.h"
 #include "PoseSearch/PoseSearchInteractionAsset.h"
 #include "PoseSearch/PoseSearchInteractionLibrary.h"
 #include "PoseSearch/PoseSearchSchema.h"
-#include "Animation/TrajectoryTypes.h"
+#include "PoseSearchInteractionAlignment.h"
+#include "RigVMCore/RigVMRegistry.h"
 
 namespace UE::UAF::PoseSearch
 {
@@ -47,10 +47,14 @@ public:
 		RigVMRegistry.RegisterStructTypes(AllowedStructTypes);
 
 		IModularFeatures::Get().RegisterModularFeature(UE::PoseSearch::IInteractionIslandDependency::FeatureName, &UE::PoseSearch::FAnimNextInteractionIslandDependency::ModularFeature);
+
+		UE::UAF::FEvaluationNotifiesTrait::RegisterEvaluationHandler(UNotifyState_PoseSearchInteractionAlignment::StaticClass(), FEvaluationNotify_PoseSearchInteractionAlignment::StaticStruct());
 	}
 
 	virtual void ShutdownModule() override
 	{
+		UE::UAF::FEvaluationNotifiesTrait::UnregisterEvaluationHandler(UNotifyState_PoseSearchInteractionAlignment::StaticClass());
+
 		IModularFeatures::Get().UnregisterModularFeature(UE::PoseSearch::IInteractionIslandDependency::FeatureName, &UE::PoseSearch::FAnimNextInteractionIslandDependency::ModularFeature);
 	}
 };
diff --git a/Engine/Plugins/Experimental/UAF/UAFPoseSearch/Source/UAFPoseSearch/Private/MotionMatchingTrait.cpp b/Engine/Plugins/Experimental/UAF/UAFPoseSearch/Source/UAFPoseSearch/Private/MotionMatchingTrait.cpp
--- a/Engine/Plugins/Experimental/UAF/UAFPoseSearch/Source/UAFPoseSearch/Private/MotionMatchingTrait.cpp
+++ b/Engine/Plugins/Experimental/UAF/UAFPoseSearch/Source/UAFPoseSearch/Private/MotionMatchingTrait.cpp
@@ -4,6 +4,7 @@
 #include "Animation/AnimRootMotionProvider.h"
 #include "Animation/AnimSequence.h"
 #include "Animation/BlendSpace.h"
+#include "AnimNext/EvaluationNotifiesTrait.h"
 #include "AnimNextAnimGraphSettings.h"
 #include "Component/AnimNextComponent.h"
 #include "EvaluationVM/EvaluationVM.h"
@@ -317,6 +318,12 @@ PRAGMA_ENABLE_DEPRECATION_WARNINGS
 				if (!bSkipBlend)
 				{
 					FAnimNextFactoryParams FactoryParams = FAnimGraphFactory::GetDefaultParamsForClass(AnimationAsset->GetClass());
+
+					// currently the default params for sequences and blend spaces lack the FEvaluationNotifiesTrait, necessary trait to be able to evaluate 
+					// FEvaluationNotify_BaseInstance(s) during trait post evaluation task (FAnimNextEvaluationNotifiesTask::Execute).
+					// Since it's required for MM interaction warp aligment, and there's no data driven way to do so, we have to hard code it here for now.
+					FactoryParams.PushPublicTraitStruct<FEvaluationNotifiesTraitSharedData>();
+
 					if (bIsBlendSpace)
 					{
 						FSynchronizeUsingGroupsData SynchronizeUsingGroups;
@@ -534,41 +541,43 @@ PRAGMA_ENABLE_DEPRECATION_WARNINGS
 		}
 #endif // WITH_EDITOR
 
-		// if it's not a multi character interaction we can skip the warping logic entirely
-		const FPoseSearchBlueprintResult& SearchResult = InstanceData->MotionMatchingState.SearchResult;
-		if (SearchResult.bIsInteraction)
+		const FSharedData* SharedData = Binding.GetSharedData<FSharedData>();
+		check(SharedData);
+		if (SharedData->GetbEnableWarp(Binding))
 		{
-			if (const UMultiAnimAsset* MultiAnimAsset = Cast<UMultiAnimAsset>(SearchResult.SelectedAnim))
+			// if it's not a multi character interaction we can skip the warping logic entirely
+			const FPoseSearchBlueprintResult& SearchResult = InstanceData->MotionMatchingState.SearchResult;
+			if (SearchResult.bIsInteraction)
 			{
-				const int32 CurrentResultRoleIndex = GetRoleIndex(MultiAnimAsset, SearchResult.Role);
-				if (CurrentResultRoleIndex != INDEX_NONE)
+				if (const UMultiAnimAsset* MultiAnimAsset = Cast<UMultiAnimAsset>(SearchResult.SelectedAnim))
 				{
-					// @todo: WIP hacky, non thread safe (unless proper tick dependencies are in place) way to retrieve the mesh transform until we find a better way
-					if (const FAnimNextModuleInstance* ModuleInstance = Context.GetRootGraphInstance().GetModuleInstance())
+					const int32 CurrentResultRoleIndex = GetRoleIndex(MultiAnimAsset, SearchResult.Role);
+					if (CurrentResultRoleIndex != INDEX_NONE)
 					{
-						if (const UAnimNextComponent* AnimNextComponent = Cast<UAnimNextComponent>(ModuleInstance->GetObject()))
+						// @todo: WIP hacky, non thread safe (unless proper tick dependencies are in place) way to retrieve the mesh transform until we find a better way
+						if (const FAnimNextModuleInstance* ModuleInstance = Context.GetRootGraphInstance().GetModuleInstance())
 						{
-							const AActor* Actor = AnimNextComponent->GetOwner();
-							check(Actor);
-							if (const USkeletalMeshComponent* SkeletalMeshComponent = Actor->GetComponentByClass<USkeletalMeshComponent>())
+							if (const UAnimNextComponent* AnimNextComponent = Cast<UAnimNextComponent>(ModuleInstance->GetObject()))
 							{
-								const FSharedData* SharedData = Binding.GetSharedData<FSharedData>();
-								check(SharedData);
-
-								FAnimNextMotionMatchingTask Task;
-								Task.ComponentTransform = SkeletalMeshComponent->GetComponentTransform();
-								Task.InstanceData = InstanceData;
-								Task.CurrentResultRoleIndex = CurrentResultRoleIndex;
-								Task.bWarpUsingRootBone = SharedData->GetbWarpUsingRootBone(Binding);
-								Task.WarpingRotationRatio = SharedData->GetWarpingRotationRatio(Binding);
-								Task.WarpingTranslationRatio = SharedData->GetWarpingTranslationRatio(Binding);
-								Task.WarpingRotationCurveName = SharedData->GetWarpingRotationCurveName(Binding);
-								Task.WarpingTranslationCurveName = SharedData->GetWarpingTranslationCurveName(Binding);
+								const AActor* Actor = AnimNextComponent->GetOwner();
+								check(Actor);
+								if (const USkeletalMeshComponent* SkeletalMeshComponent = Actor->GetComponentByClass<USkeletalMeshComponent>())
+								{
+									FAnimNextMotionMatchingTask Task;
+									Task.ComponentTransform = SkeletalMeshComponent->GetComponentTransform();
+									Task.InstanceData = InstanceData;
+									Task.CurrentResultRoleIndex = CurrentResultRoleIndex;
+									Task.bWarpUsingRootBone = SharedData->GetbWarpUsingRootBone(Binding);
+									Task.WarpingRotationRatio = SharedData->GetWarpingRotationRatio(Binding);
+									Task.WarpingTranslationRatio = SharedData->GetWarpingTranslationRatio(Binding);
+									Task.WarpingRotationCurveName = SharedData->GetWarpingRotationCurveName(Binding);
+									Task.WarpingTranslationCurveName = SharedData->GetWarpingTranslationCurveName(Binding);
 #if ENABLE_ANIM_DEBUG
-								// Debug Object for VisualLogger
-								Task.HostObject = Context.GetHostObject();
+									// Debug Object for VisualLogger
+									Task.HostObject = Context.GetHostObject();
 #endif // ENABLE_ANIM_DEBUG 
-								Context.AppendTask(Task);
+									Context.AppendTask(Task);
+								}
 							}
 						}
 					}
@@ -665,7 +674,6 @@ void FAnimNextMotionMatchingTask::Execute(UE::UAF::FEvaluationVM& VM) const
 		DeltaAlignment.SetTranslation(DeltaTranslation);
 	}
 
-
 	const FTransform DeltaAlignmentWithRootMotion = DeltaAlignment * RootMotionDelta;
 
 	RootMotionProvider->OverrideRootMotion(DeltaAlignmentWithRootMotion, Keyframe->Get()->Attributes);
diff --git a/Engine/Plugins/Experimental/UAF/UAFPoseSearch/Source/UAFPoseSearch/Private/PoseSearchInteractionAlignment.cpp b/Engine/Plugins/Experimental/UAF/UAFPoseSearch/Source/UAFPoseSearch/Private/PoseSearchInteractionAlignment.cpp
new file mode 100644
--- /dev/null
+++ b/Engine/Plugins/Experimental/UAF/UAFPoseSearch/Source/UAFPoseSearch/Private/PoseSearchInteractionAlignment.cpp
@@ -0,0 +1,93 @@
+// Copyright Epic Games, Inc. All Rights Reserved.
+
+#include "PoseSearchInteractionAlignment.h"
+#include "Animation/AnimRootMotionProvider.h"
+#include "EvaluationVM/EvaluationVM.h"
+#include "PoseSearch/PoseSearchContext.h"
+#include "PoseSearch/PoseSearchInteractionSubsystem.h"
+#include "PoseSearch/PoseSearchInteractionUtils.h"
+#include "VisualLogger/VisualLogger.h"
+
+void FEvaluationNotify_PoseSearchInteractionAlignment::Start()
+{
+	WarpStart = FTransform::Identity;
+	WarpEnd = FTransform::Identity;
+	WarpEndTime = 0.f;
+	WarpCurrentTime = 0.f;
+	bFirstFrame = true;
+}
+
+void FEvaluationNotify_PoseSearchInteractionAlignment::Update(UE::UAF::FEvaluationNotifiesTrait::FInstanceData& InstanceData, UE::UAF::FEvaluationVM& VM)
+{
+	using namespace UE::Anim;
+	using namespace UE::PoseSearch;
+
+	UNotifyState_PoseSearchInteractionAlignment* PoseSearchInteractionAlignment = CastChecked<UNotifyState_PoseSearchInteractionAlignment>(AnimNotify);
+
+	if (const TUniquePtr<UE::UAF::FKeyframeState>* Keyframe = VM.PeekValue<TUniquePtr<UE::UAF::FKeyframeState>>(UE::UAF::KEYFRAME_STACK_NAME, 0))
+	{
+		const UObject* AnimContext = InstanceData.HostObject;
+
+		const IAnimRootMotionProvider* RootMotionProvider = IAnimRootMotionProvider::Get();
+		UPoseSearchInteractionSubsystem* InteractionSubsystem = UPoseSearchInteractionSubsystem::GetSubsystem_AnyThread(AnimContext);
+
+		if (RootMotionProvider && InteractionSubsystem)
+		{
+			// thread safety note: because multi character motion matching adds a tick dependecy between character and UAF,
+			// using GetContextTransform(AnimContext, false) is thread safe
+			const FTransform ContextTransform = GetContextTransform(AnimContext, false);
+
+			if (bFirstFrame)
+			{
+				FPoseSearchBlueprintResult Result;
+				InteractionSubsystem->GetResult_AnyThread(AnimContext, Result);
+
+				if (Result.SelectedAnim)
+				{
+					check(Result.RoleIndex >= 0 && Result.RoleIndex < Result.AnimContexts.Num());
+
+					WarpStart = ContextTransform;
+
+					WarpEndTime = EndTime - CurrentTime;
+
+					const float TimeOffset = EndTime - Result.SelectedTime;
+
+					// @todo: this should be a property / parameter
+					static bool bWarpUsingRootBone = false;
+
+					TArray<FTransform, TInlineAllocator<UE::PoseSearch::PreallocatedRolesNum>> FullAlignedTransforms;
+					FullAlignedTransforms.SetNum(Result.AnimContexts.Num());
+					UE::PoseSearch::CalculateFullAlignedTransforms(Result, Result.SelectedTime, TimeOffset, bWarpUsingRootBone, FullAlignedTransforms);
+
+					WarpEnd = FullAlignedTransforms[Result.RoleIndex];
+				}
+			}
+
+			if (WarpEndTime > UE_KINDA_SMALL_NUMBER)
+			{
+				const float Alpha = FMath::Min(WarpCurrentTime / WarpEndTime, 1.f);
+				FTransform WarpCurrent;
+				WarpCurrent.Blend(WarpStart, WarpEnd, Alpha);
+
+				const FTransform RootMotion = WarpCurrent.GetRelativeTransform(ContextTransform);
+				RootMotionProvider->OverrideRootMotion(RootMotion, Keyframe->Get()->Attributes);
+
+				//UE_VLOG_SPHERE(AnimContext, "MMIAlignment", Display, WarpStart.GetLocation(), 0, FColor::Red, TEXT(""));
+				UE_VLOG_SEGMENT(AnimContext, "MMIAlignment", Display, WarpStart.GetLocation(), WarpStart.GetLocation() + WarpStart.GetRotation().RotateVector(FVector::XAxisVector) * 30.f, FColor::Red, TEXT(""));
+				//UE_VLOG_SPHERE(AnimContext, "MMIAlignment", Display, WarpCurrent.GetLocation(), 0, FColor::Green, TEXT(""));
+				UE_VLOG_SEGMENT(AnimContext, "MMIAlignment", Display, WarpCurrent.GetLocation(), WarpCurrent.GetLocation() + WarpCurrent.GetRotation().RotateVector(FVector::XAxisVector) * 30.f, FColor::Green, TEXT(""));
+				//UE_VLOG_SPHERE(AnimContext, "MMIAlignment", Display, WarpEnd.GetLocation(), 0, FColor::Blue, TEXT(""));
+				UE_VLOG_SEGMENT(AnimContext, "MMIAlignment", Display, WarpEnd.GetLocation(), WarpEnd.GetLocation() + WarpEnd.GetRotation().RotateVector(FVector::XAxisVector) * 30.f, FColor::Blue, TEXT(""));
+			}
+		}
+	}
+
+	WarpCurrentTime += InstanceData.DeltaTime;
+	bFirstFrame = false;
+}
+
+void FEvaluationNotify_PoseSearchInteractionAlignment::End(UE::UAF::FEvaluationNotifiesTrait::FInstanceData& InstanceData)
+{
+	static int End = 0;
+	End++;
+}
\ No newline at end of file
diff --git a/Engine/Plugins/Experimental/UAF/UAFPoseSearch/Source/UAFPoseSearch/Private/PoseSearchInteractionAlignment.h b/Engine/Plugins/Experimental/UAF/UAFPoseSearch/Source/UAFPoseSearch/Private/PoseSearchInteractionAlignment.h
new file mode 100644
--- /dev/null
+++ b/Engine/Plugins/Experimental/UAF/UAFPoseSearch/Source/UAFPoseSearch/Private/PoseSearchInteractionAlignment.h
@@ -0,0 +1,31 @@
+// Copyright Epic Games, Inc. All Rights Reserved.
+
+#pragma once
+
+#include "AnimNext/EvaluationNotifiesTrait.h"
+#include "Animation/AnimNotifies/AnimNotifyState.h"
+#include "PoseSearchInteractionAlignment.generated.h"
+
+USTRUCT()
+struct FEvaluationNotify_PoseSearchInteractionAlignment : public FEvaluationNotify_BaseInstance
+{
+	GENERATED_BODY()
+
+	virtual void Start() override;
+	virtual void Update(UE::UAF::FEvaluationNotifiesTrait::FInstanceData& InstanceData, UE::UAF::FEvaluationVM& VM) override;
+	virtual void End(UE::UAF::FEvaluationNotifiesTrait::FInstanceData& InstanceData) override;
+
+	FTransform WarpStart = FTransform::Identity;
+	FTransform WarpEnd = FTransform::Identity;
+	float WarpEndTime = 0.f;
+	float WarpCurrentTime = 0.f;
+	bool bFirstFrame = true;
+};
+
+UCLASS(MinimalAPI, BlueprintType, DisplayName="MMI Alignment")
+class UNotifyState_PoseSearchInteractionAlignment : public UAnimNotifyState
+{
+	GENERATED_BODY()
+
+public:
+};
\ No newline at end of file
diff --git a/Engine/Plugins/Experimental/UAF/UAFPoseSearch/Source/UAFPoseSearch/UAFPoseSearch.Build.cs b/Engine/Plugins/Experimental/UAF/UAFPoseSearch/Source/UAFPoseSearch/UAFPoseSearch.Build.cs
--- a/Engine/Plugins/Experimental/UAF/UAFPoseSearch/Source/UAFPoseSearch/UAFPoseSearch.Build.cs
+++ b/Engine/Plugins/Experimental/UAF/UAFPoseSearch/Source/UAFPoseSearch/UAFPoseSearch.Build.cs
@@ -9,16 +9,17 @@ namespace UnrealBuildTool.Rules
 			PublicDependencyModuleNames.AddRange(
 				new string[]
 				{
-					"UAF",
-					"UAFAnimGraph",
+					"Chooser",
 					"Core",
 					"CoreUObject",
 					"Engine",
-					"HierarchyTableRuntime",
+					"EvaluationNotifiesRuntime",
 					"HierarchyTableAnimationRuntime",
+					"HierarchyTableRuntime",
 					"PoseSearch",
-					"Chooser",
 					"RigVM",
+					"UAF",
+					"UAFAnimGraph"
 				}
 			);
 		}
diff --git a/Engine/Plugins/Experimental/UAF/UAFPoseSearch/UAFPoseSearch.uplugin b/Engine/Plugins/Experimental/UAF/UAFPoseSearch/UAFPoseSearch.uplugin
--- a/Engine/Plugins/Experimental/UAF/UAFPoseSearch/UAFPoseSearch.uplugin
+++ b/Engine/Plugins/Experimental/UAF/UAFPoseSearch/UAFPoseSearch.uplugin
@@ -18,11 +18,11 @@
 	"Plugins":
 	[
 		{
-			"Name": "UAF",
+			"Name": "Chooser",
 			"Enabled": true
 		},
 		{
-			"Name": "UAFAnimGraph",
+			"Name": "EvaluationNotifies",
 			"Enabled": true
 		},
 		{
@@ -30,9 +30,13 @@
 			"Enabled": true
 		},
 		{
-			"Name": "Chooser",
+			"Name": "UAF",
+			"Enabled": true
+		},
+		{
+			"Name": "UAFAnimGraph",
 			"Enabled": true
-		}
+		}		
 	],
 	"Modules" :
 	[
