UAFPoseSearch - MMI alignment warping and steering improvements

#Synced-CL 48544215

diff --git a/Engine/Plugins/Experimental/UAF/UAFPoseSearch/Source/UAFPoseSearch/Private/Module.cpp b/Engine/Plugins/Experimental/UAF/UAFPoseSearch/Source/UAFPoseSearch/Private/Module.cpp
--- a/Engine/Plugins/Experimental/UAF/UAFPoseSearch/Source/UAFPoseSearch/Private/Module.cpp
+++ b/Engine/Plugins/Experimental/UAF/UAFPoseSearch/Source/UAFPoseSearch/Private/Module.cpp
@@ -13,6 +13,7 @@
 #include "PoseSearch/PoseSearchInteractionLibrary.h"
 #include "PoseSearch/PoseSearchSchema.h"
 #include "PoseSearchInteractionAlignment.h"
+#include "PoseSearchSteerAlongTrajectory.h"
 #include "RigVMCore/RigVMRegistry.h"
 
 namespace UE::UAF::PoseSearch
@@ -49,10 +50,12 @@ public:
 		IModularFeatures::Get().RegisterModularFeature(UE::PoseSearch::IInteractionIslandDependency::FeatureName, &UE::PoseSearch::FAnimNextInteractionIslandDependency::ModularFeature);
 
 		UE::UAF::FEvaluationNotifiesTrait::RegisterEvaluationHandler(UNotifyState_PoseSearchInteractionAlignment::StaticClass(), FEvaluationNotify_PoseSearchInteractionAlignment::StaticStruct());
+		UE::UAF::FEvaluationNotifiesTrait::RegisterEvaluationHandler(UNotifyState_PoseSearchSteerAlongTrajectory::StaticClass(), FEvaluationNotify_PoseSearchSteerAlongTrajectory::StaticStruct());
 	}
 
 	virtual void ShutdownModule() override
 	{
+		UE::UAF::FEvaluationNotifiesTrait::UnregisterEvaluationHandler(UNotifyState_PoseSearchSteerAlongTrajectory::StaticClass());
 		UE::UAF::FEvaluationNotifiesTrait::UnregisterEvaluationHandler(UNotifyState_PoseSearchInteractionAlignment::StaticClass());
 
 		IModularFeatures::Get().UnregisterModularFeature(UE::PoseSearch::IInteractionIslandDependency::FeatureName, &UE::PoseSearch::FAnimNextInteractionIslandDependency::ModularFeature);
diff --git a/Engine/Plugins/Experimental/UAF/UAFPoseSearch/Source/UAFPoseSearch/Private/PoseSearchInteractionAlignment.cpp b/Engine/Plugins/Experimental/UAF/UAFPoseSearch/Source/UAFPoseSearch/Private/PoseSearchInteractionAlignment.cpp
--- a/Engine/Plugins/Experimental/UAF/UAFPoseSearch/Source/UAFPoseSearch/Private/PoseSearchInteractionAlignment.cpp
+++ b/Engine/Plugins/Experimental/UAF/UAFPoseSearch/Source/UAFPoseSearch/Private/PoseSearchInteractionAlignment.cpp
@@ -21,10 +21,9 @@ void FEvaluationNotify_PoseSearchInteractionAlignment::Update(UE::UAF::FEvaluati
 {
 	using namespace UE::Anim;
 	using namespace UE::PoseSearch;
+	using namespace UE::UAF;
 
-	UNotifyState_PoseSearchInteractionAlignment* PoseSearchInteractionAlignment = CastChecked<UNotifyState_PoseSearchInteractionAlignment>(AnimNotify);
-
-	if (const TUniquePtr<UE::UAF::FKeyframeState>* Keyframe = VM.PeekValue<TUniquePtr<UE::UAF::FKeyframeState>>(UE::UAF::KEYFRAME_STACK_NAME, 0))
+	if (const TUniquePtr<FKeyframeState>* Keyframe = VM.PeekValue<TUniquePtr<FKeyframeState>>(KEYFRAME_STACK_NAME, 0))
 	{
 		const UObject* AnimContext = InstanceData.HostObject;
 
@@ -72,11 +71,8 @@ void FEvaluationNotify_PoseSearchInteractionAlignment::Update(UE::UAF::FEvaluati
 				const FTransform RootMotion = WarpCurrent.GetRelativeTransform(ContextTransform);
 				RootMotionProvider->OverrideRootMotion(RootMotion, Keyframe->Get()->Attributes);
 
-				//UE_VLOG_SPHERE(AnimContext, "MMIAlignment", Display, WarpStart.GetLocation(), 0, FColor::Red, TEXT(""));
 				UE_VLOG_SEGMENT(AnimContext, "MMIAlignment", Display, WarpStart.GetLocation(), WarpStart.GetLocation() + WarpStart.GetRotation().RotateVector(FVector::XAxisVector) * 30.f, FColor::Red, TEXT(""));
-				//UE_VLOG_SPHERE(AnimContext, "MMIAlignment", Display, WarpCurrent.GetLocation(), 0, FColor::Green, TEXT(""));
 				UE_VLOG_SEGMENT(AnimContext, "MMIAlignment", Display, WarpCurrent.GetLocation(), WarpCurrent.GetLocation() + WarpCurrent.GetRotation().RotateVector(FVector::XAxisVector) * 30.f, FColor::Green, TEXT(""));
-				//UE_VLOG_SPHERE(AnimContext, "MMIAlignment", Display, WarpEnd.GetLocation(), 0, FColor::Blue, TEXT(""));
 				UE_VLOG_SEGMENT(AnimContext, "MMIAlignment", Display, WarpEnd.GetLocation(), WarpEnd.GetLocation() + WarpEnd.GetRotation().RotateVector(FVector::XAxisVector) * 30.f, FColor::Blue, TEXT(""));
 			}
 		}
@@ -85,9 +81,3 @@ void FEvaluationNotify_PoseSearchInteractionAlignment::Update(UE::UAF::FEvaluati
 	WarpCurrentTime += InstanceData.DeltaTime;
 	bFirstFrame = false;
 }
-
-void FEvaluationNotify_PoseSearchInteractionAlignment::End(UE::UAF::FEvaluationNotifiesTrait::FInstanceData& InstanceData)
-{
-	static int End = 0;
-	End++;
-}
\ No newline at end of file
diff --git a/Engine/Plugins/Experimental/UAF/UAFPoseSearch/Source/UAFPoseSearch/Private/PoseSearchInteractionAlignment.h b/Engine/Plugins/Experimental/UAF/UAFPoseSearch/Source/UAFPoseSearch/Private/PoseSearchInteractionAlignment.h
--- a/Engine/Plugins/Experimental/UAF/UAFPoseSearch/Source/UAFPoseSearch/Private/PoseSearchInteractionAlignment.h
+++ b/Engine/Plugins/Experimental/UAF/UAFPoseSearch/Source/UAFPoseSearch/Private/PoseSearchInteractionAlignment.h
@@ -6,14 +6,13 @@
 #include "Animation/AnimNotifies/AnimNotifyState.h"
 #include "PoseSearchInteractionAlignment.generated.h"
 
-USTRUCT()
+USTRUCT(Experimental)
 struct FEvaluationNotify_PoseSearchInteractionAlignment : public FEvaluationNotify_BaseInstance
 {
 	GENERATED_BODY()
 
 	virtual void Start() override;
 	virtual void Update(UE::UAF::FEvaluationNotifiesTrait::FInstanceData& InstanceData, UE::UAF::FEvaluationVM& VM) override;
-	virtual void End(UE::UAF::FEvaluationNotifiesTrait::FInstanceData& InstanceData) override;
 
 	FTransform WarpStart = FTransform::Identity;
 	FTransform WarpEnd = FTransform::Identity;
@@ -22,10 +21,8 @@ struct FEvaluationNotify_PoseSearchInteractionAlignment : public FEvaluationNoti
 	bool bFirstFrame = true;
 };
 
-UCLASS(MinimalAPI, BlueprintType, DisplayName="MMI Alignment")
+UCLASS(Experimental, MinimalAPI, BlueprintType, DisplayName="MMI Alignment")
 class UNotifyState_PoseSearchInteractionAlignment : public UAnimNotifyState
 {
 	GENERATED_BODY()
-
-public:
 };
\ No newline at end of file
diff --git a/Engine/Plugins/Experimental/UAF/UAFPoseSearch/Source/UAFPoseSearch/Private/PoseSearchSteerAlongTrajectory.cpp b/Engine/Plugins/Experimental/UAF/UAFPoseSearch/Source/UAFPoseSearch/Private/PoseSearchSteerAlongTrajectory.cpp
new file mode 100644
--- /dev/null
+++ b/Engine/Plugins/Experimental/UAF/UAFPoseSearch/Source/UAFPoseSearch/Private/PoseSearchSteerAlongTrajectory.cpp
@@ -0,0 +1,93 @@
+// Copyright Epic Games, Inc. All Rights Reserved.
+
+#include "PoseSearchSteerAlongTrajectory.h"
+#include "Animation/AnimRootMotionProvider.h"
+#include "Animation/SpringMath.h"
+#include "Animation/TrajectoryTypes.h"
+#include "EvaluationVM/EvaluationVM.h"
+#include "PoseSearch/PoseSearchContext.h"
+#include "VisualLogger/VisualLogger.h"
+#include "Module/AnimNextModuleInstance.h"
+
+namespace UE::UAF::PoseSearch
+{
+	template<typename ValueType>
+	static EPropertyBagResult AccessVariable(const FUAFAssetInstance* Instance, const FAnimNextVariableReference& Variable, TFunctionRef<void(ValueType&)> Function)
+	{
+		while (Instance)
+		{
+			if (EPropertyBagResult::Success == Instance->AccessVariable(Variable, Function))
+			{
+				return EPropertyBagResult::Success;
+			}
+			Instance = Instance->GetHost();
+		}
+		return EPropertyBagResult::PropertyNotFound;
+	}
+} // namespace UE::UAF::PoseSearch
+
+void FEvaluationNotify_PoseSearchSteerAlongTrajectory::Start()
+{
+	AngularVelocity = FVector::ZeroVector;
+}
+
+void FEvaluationNotify_PoseSearchSteerAlongTrajectory::Update(UE::UAF::FEvaluationNotifiesTrait::FInstanceData& InstanceData, UE::UAF::FEvaluationVM& VM)
+{
+	using namespace UE::Anim;
+	using namespace UE::PoseSearch;
+	using namespace UE::UAF::PoseSearch;
+	using namespace UE::UAF;
+
+	if (const TUniquePtr<FKeyframeState>* Keyframe = VM.PeekValue<TUniquePtr<FKeyframeState>>(KEYFRAME_STACK_NAME, 0))
+	{
+		if (const IAnimRootMotionProvider* RootMotionProvider = IAnimRootMotionProvider::Get())
+		{
+			const UNotifyState_PoseSearchSteerAlongTrajectory* SteerAlongTrajectory = CastChecked<UNotifyState_PoseSearchSteerAlongTrajectory>(AnimNotify);
+
+			FVector TrajectoryDirection;
+			auto TransformTrajectoryAccessor = [SteerAlongTrajectory, &TrajectoryDirection](FTransformTrajectory& TransformTrajectory) -> void
+				{
+					const FTransformTrajectorySample A = TransformTrajectory.GetSampleAtTime(SteerAlongTrajectory->StartTrajectorySampleTime, true);
+					const FTransformTrajectorySample B = TransformTrajectory.GetSampleAtTime(SteerAlongTrajectory->EndTrajectorySampleTime, true);
+					TrajectoryDirection = B.Position - A.Position;
+				};
+
+			TFunctionRef<void(FTransformTrajectory&)> TransformTrajectoryAccessorRef = TransformTrajectoryAccessor;
+			if (EPropertyBagResult::Success == AccessVariable(InstanceData.Instance, SteerAlongTrajectory->Trajectory, TransformTrajectoryAccessorRef))
+			{
+				if (TrajectoryDirection.Normalize(UE_KINDA_SMALL_NUMBER))
+				{
+					FStackAttributeContainer& Attributes = Keyframe->Get()->Attributes;
+
+					const UObject* AnimContext = InstanceData.HostObject;
+
+					// @todo: we can also use TransformTrajectory.GetSampleAtTime(0) to get the ContextTransform
+					// thread safety note: because multi character motion matching adds a tick dependecy between character and UAF,
+					// using GetContextTransform(AnimContext, false) is thread safe
+					const FTransform ContextTransform = GetContextTransform(AnimContext, false);
+					
+					FTransform RootMotion;
+					RootMotionProvider->ExtractRootMotion(Attributes, RootMotion);
+
+					FVector RootMotionDirection = RootMotion.GetTranslation();
+					RootMotionDirection = ContextTransform.TransformVectorNoScale(RootMotionDirection);
+					if (RootMotionDirection.Normalize(UE_KINDA_SMALL_NUMBER))
+					{
+						const FQuat RotationDelta = FQuat::FindBetweenVectors(RootMotionDirection, TrajectoryDirection);
+
+						FQuat RootMotionRotation = RootMotion.GetRotation();
+						const FQuat RootMotionTargetRotation = RootMotion.GetRotation() * RotationDelta;
+
+						SpringMath::CriticalSpringDamperQuat(RootMotionRotation, AngularVelocity, RootMotionTargetRotation, SteerAlongTrajectory->SmoothingTime, InstanceData.DeltaTime);
+
+						RootMotion.SetRotation(RootMotionRotation);
+						RootMotionProvider->OverrideRootMotion(RootMotion, Attributes);
+
+						UE_VLOG_SEGMENT(AnimContext, "MMISteerAlongTrajectory", Display, ContextTransform.GetLocation(), ContextTransform.GetLocation() + TrajectoryDirection * 60.f, FColor::Blue, TEXT(""));
+						UE_VLOG_SEGMENT(AnimContext, "MMISteerAlongTrajectory", Display, ContextTransform.GetLocation(), ContextTransform.GetLocation() + RootMotionDirection * 60.f, FColor::Green, TEXT(""));
+					}
+				}
+			}
+		}
+	}
+}
\ No newline at end of file
diff --git a/Engine/Plugins/Experimental/UAF/UAFPoseSearch/Source/UAFPoseSearch/Private/PoseSearchSteerAlongTrajectory.h b/Engine/Plugins/Experimental/UAF/UAFPoseSearch/Source/UAFPoseSearch/Private/PoseSearchSteerAlongTrajectory.h
new file mode 100644
--- /dev/null
+++ b/Engine/Plugins/Experimental/UAF/UAFPoseSearch/Source/UAFPoseSearch/Private/PoseSearchSteerAlongTrajectory.h
@@ -0,0 +1,37 @@
+// Copyright Epic Games, Inc. All Rights Reserved.
+
+#pragma once
+
+#include "AnimNext/EvaluationNotifiesTrait.h"
+#include "Animation/AnimNotifies/AnimNotifyState.h"
+#include "PoseSearchSteerAlongTrajectory.generated.h"
+
+USTRUCT(Experimental)
+struct FEvaluationNotify_PoseSearchSteerAlongTrajectory : public FEvaluationNotify_BaseInstance
+{
+	GENERATED_BODY()
+
+	virtual void Start() override;
+	virtual void Update(UE::UAF::FEvaluationNotifiesTrait::FInstanceData& InstanceData, UE::UAF::FEvaluationVM& VM) override;
+	
+	FVector AngularVelocity = FVector::ZeroVector;
+};
+
+UCLASS(Experimental, MinimalAPI, BlueprintType, DisplayName="MMI Steer Along Trajectory")
+class UNotifyState_PoseSearchSteerAlongTrajectory : public UAnimNotifyState
+{
+	GENERATED_BODY()
+
+public:
+	UPROPERTY(EditAnywhere, Category="Steering")
+	FAnimNextVariableReference Trajectory;
+
+	UPROPERTY(EditAnywhere, Category="Steering")
+	float StartTrajectorySampleTime = 0.5f;
+
+	UPROPERTY(EditAnywhere, Category="Steering")
+	float EndTrajectorySampleTime = 1.f;
+
+	UPROPERTY(EditAnywhere, Category="Steering")
+	float SmoothingTime = 1.f;
+};
\ No newline at end of file
