Fix UAF editor crash in pose history trait

#Synced-CL 48851180

diff --git a/Engine/Plugins/Experimental/UAF/UAFPoseSearch/Source/UAFPoseSearch/Internal/HistoryCollectorTraitData.h b/Engine/Plugins/Experimental/UAF/UAFPoseSearch/Source/UAFPoseSearch/Internal/HistoryCollectorTraitData.h
--- a/Engine/Plugins/Experimental/UAF/UAFPoseSearch/Source/UAFPoseSearch/Internal/HistoryCollectorTraitData.h
+++ b/Engine/Plugins/Experimental/UAF/UAFPoseSearch/Source/UAFPoseSearch/Internal/HistoryCollectorTraitData.h
@@ -83,7 +83,7 @@ struct FAnimNextHistoryCollectorTraitSharedData : public FAnimNextTraitSharedDat
 
 	// Property to store a reference to the PoseHistory struct in
 	UPROPERTY(EditAnywhere, Category = Experimental)
-	FAnimNextVariableReference PoseHistoryReferenceVariable;
+	FName PoseHistoryReferenceVariable;
 
 	// if bGenerateTrajectory is true, TrajectoryData contains the tuning parameters to generate the trajectory
 	//UPROPERTY(EditAnywhere, Category = Experimental, meta=(EditCondition="bGenerateTrajectory", EditConditionHides))
diff --git a/Engine/Plugins/Experimental/UAF/UAFPoseSearch/Source/UAFPoseSearch/Private/HistoryCollectorTrait.cpp b/Engine/Plugins/Experimental/UAF/UAFPoseSearch/Source/UAFPoseSearch/Private/HistoryCollectorTrait.cpp
--- a/Engine/Plugins/Experimental/UAF/UAFPoseSearch/Source/UAFPoseSearch/Private/HistoryCollectorTrait.cpp
+++ b/Engine/Plugins/Experimental/UAF/UAFPoseSearch/Source/UAFPoseSearch/Private/HistoryCollectorTrait.cpp
@@ -283,12 +283,14 @@ namespace UE::UAF
 		}
 #endif // WITH_EDITOR
 
-		const FAnimNextVariableReference& PoseHistoryReferenceVariable = SharedData->GetPoseHistoryReferenceVariable(Binding);
-		if (PoseHistoryReferenceVariable.IsValid())
+		const FName& PoseHistoryReferenceVariableName = SharedData->GetPoseHistoryReferenceVariable(Binding);
+		if (PoseHistoryReferenceVariableName != NAME_None)
 		{
 			FPoseHistoryReference Reference;
 			Reference.PoseHistory = InstanceData->PoseHistoryPtr;
-			SetVariableValue(&Binding.GetTraitPtr().GetNodeInstance()->GetOwner(), PoseHistoryReferenceVariable, Reference);
+			PRAGMA_DISABLE_DEPRECATION_WARNINGS
+			SetVariableValue(&Binding.GetTraitPtr().GetNodeInstance()->GetOwner(), FAnimNextVariableReference(PoseHistoryReferenceVariableName), Reference);
+			PRAGMA_ENABLE_DEPRECATION_WARNINGS
 		}
 		
 		check(InstanceData->PoseHistoryPtr.IsValid());
