Pose Search: Register PoseSearchEvent as a RigVM type

Prevents ensure when linking to the motion matching trait

#Synced-CL 50295456

diff --git a/Engine/Plugins/Experimental/UAF/UAFPoseSearch/Source/UAFPoseSearch/Private/Module.cpp b/Engine/Plugins/Experimental/UAF/UAFPoseSearch/Source/UAFPoseSearch/Private/Module.cpp
--- a/Engine/Plugins/Experimental/UAF/UAFPoseSearch/Source/UAFPoseSearch/Private/Module.cpp
+++ b/Engine/Plugins/Experimental/UAF/UAFPoseSearch/Source/UAFPoseSearch/Private/Module.cpp
@@ -41,6 +41,7 @@ public:
 			FPoseSearchInteractionAssetItem::StaticStruct(),
 			FPoseSearchInteractionAvailability::StaticStruct(),
 			FPoseHistoryReference::StaticStruct(),
+			FPoseSearchEvent::StaticStruct(),
 		};
 
 		FRigVMRegistry& RigVMRegistry = FRigVMRegistry::Get();
