PoseSearch - FHistoryCollectorTrait improvements to get closer to parity with FAnimNode_PoseSearchHistoryCollector_Base

#Synced-CL 50317418

diff --git a/Engine/Plugins/Animation/PoseSearch/Source/Runtime/Private/AnimNode_PoseSearchHistoryCollector.cpp b/Engine/Plugins/Animation/PoseSearch/Source/Runtime/Private/AnimNode_PoseSearchHistoryCollector.cpp
--- a/Engine/Plugins/Animation/PoseSearch/Source/Runtime/Private/AnimNode_PoseSearchHistoryCollector.cpp
+++ b/Engine/Plugins/Animation/PoseSearch/Source/Runtime/Private/AnimNode_PoseSearchHistoryCollector.cpp
@@ -55,7 +55,7 @@ void FAnimNode_PoseSearchHistoryCollector_Base::Initialize_AnyThread(const FAnim
 	{
 		PoseHistoryPtr->EvaluateComponentSpace_AnyThread(0.f, ComponentSpacePoseProvider, bStoreScales,
 			RootBoneRecoveryTime, RootBoneTranslationRecoveryRatio, RootBoneRotationRecoveryRatio, true, true,
-			GetRequiredBones(Context.AnimInstanceProxy), FBlendedCurve(), MakeConstArrayView(CollectedCurves)
+			GetRequiredBones(Context.AnimInstanceProxy), FBlendedCurve(), CollectedCurves
 #if WITH_EDITORONLY_DATA
 			, GetContextTransform(AnimContext, false)
 #endif //WITH_EDITORONLY_DATA
diff --git a/Engine/Plugins/Animation/PoseSearch/Source/Runtime/Private/PoseSearchContext.cpp b/Engine/Plugins/Animation/PoseSearch/Source/Runtime/Private/PoseSearchContext.cpp
--- a/Engine/Plugins/Animation/PoseSearch/Source/Runtime/Private/PoseSearchContext.cpp
+++ b/Engine/Plugins/Animation/PoseSearch/Source/Runtime/Private/PoseSearchContext.cpp
@@ -953,7 +953,17 @@ float FSearchContext::GetSampleCurveValueInternal(float SampleTime, const FName&
 	const IPoseHistory* PoseHistory = GetPoseHistory(SampleRole);
 	check(PoseHistory);
 	
-	PoseHistory->GetCurveValueAtTime(SampleTime, CurveName, OutCurveValue);
+#if !NO_LOGGING
+	const bool bSuccess = 
+#endif // !NO_LOGGING
+		PoseHistory->GetCurveValueAtTime(SampleTime, CurveName, OutCurveValue);
+#if !NO_LOGGING
+	if (!bSuccess && !PoseHistory->IsEmpty())
+	{
+		UE_LOG(LogPoseSearch, Warning, TEXT("FSearchContext::GetSampleCurveValueInternal - Couldn't find Curve %s in the input IPoseHistory requested by %s. Consider adding it to the Pose History!"), *CurveName.ToString(), *Schema->GetName());
+	}
+#endif // !NO_LOGGING
+
 	return OutCurveValue;
 }
 
diff --git a/Engine/Plugins/Experimental/UAF/UAFPoseSearch/Source/UAFPoseSearch/Internal/HistoryCollectorTraitData.h b/Engine/Plugins/Experimental/UAF/UAFPoseSearch/Source/UAFPoseSearch/Internal/HistoryCollectorTraitData.h
--- a/Engine/Plugins/Experimental/UAF/UAFPoseSearch/Source/UAFPoseSearch/Internal/HistoryCollectorTraitData.h
+++ b/Engine/Plugins/Experimental/UAF/UAFPoseSearch/Source/UAFPoseSearch/Internal/HistoryCollectorTraitData.h
@@ -2,11 +2,12 @@
 
 #pragma once
 
-#include "CoreMinimal.h"
-#include "Animation/TrajectoryTypes.h"
-#include "TraitCore/TraitSharedData.h"
 #include "Animation/BoneReference.h"
+#include "Animation/TrajectoryTypes.h"
+#include "CoreMinimal.h"
+#include "PoseSearch/PoseSearchTrajectoryLibrary.h"
 #include "TraitCore/TraitBinding.h"
+#include "TraitCore/TraitSharedData.h"
 #include "Variables/AnimNextVariableReference.h"
 #include "HistoryCollectorTraitData.generated.h"
 
@@ -26,9 +27,8 @@ struct FAnimNextHistoryCollectorTraitSharedData : public FAnimNextTraitSharedDat
 	UPROPERTY(EditAnywhere, Category = Settings)
 	TArray<FBoneReference> CollectedBones;
 
-	// if true, the pose history will be initialized with a ref pose at the location and orientation of the AnimInstance.
 	UPROPERTY(EditAnywhere, Category = Settings)
-	bool bInitializeWithRefPose = false;
+	TArray<FName> CollectedCurves;
 
 	// Reset the pose history if it has become relevant to the graph after not being updated on previous frames.
 	UPROPERTY(EditAnywhere, Category = Settings)
@@ -86,15 +86,15 @@ struct FAnimNextHistoryCollectorTraitSharedData : public FAnimNextTraitSharedDat
 	FName PoseHistoryReferenceVariable;
 
 	// if bGenerateTrajectory is true, TrajectoryData contains the tuning parameters to generate the trajectory
-	//UPROPERTY(EditAnywhere, Category = Experimental, meta=(EditCondition="bGenerateTrajectory", EditConditionHides))
-	//FPoseSearchTrajectoryData TrajectoryData;
+	UPROPERTY(EditAnywhere, Category = Experimental, meta=(EditCondition="bGenerateTrajectory", EditConditionHides))
+	FPoseSearchTrajectoryData TrajectoryData;
 
 	// Latent pin support boilerplate
 	#define TRAIT_LATENT_PROPERTIES_ENUMERATOR(GeneratorMacro) \
 		GeneratorMacro(PoseCount) \
 		GeneratorMacro(SamplingInterval) \
 		GeneratorMacro(CollectedBones) \
-		GeneratorMacro(bInitializeWithRefPose) \
+		GeneratorMacro(CollectedCurves) \
 		GeneratorMacro(bResetOnBecomingRelevant) \
 		GeneratorMacro(bStoreScales) \
 		GeneratorMacro(RootBoneRecoveryTime) \
@@ -107,6 +107,7 @@ struct FAnimNextHistoryCollectorTraitSharedData : public FAnimNextTraitSharedDat
 		GeneratorMacro(TrajectoryPredictionCount) \
 		GeneratorMacro(PredictionSamplingInterval) \
 		GeneratorMacro(PoseHistoryReferenceVariable) \
+		GeneratorMacro(TrajectoryData) \
 
 	GENERATE_TRAIT_LATENT_PROPERTIES(FAnimNextHistoryCollectorTraitSharedData, TRAIT_LATENT_PROPERTIES_ENUMERATOR)
 	#undef TRAIT_LATENT_PROPERTIES_ENUMERATOR
diff --git a/Engine/Plugins/Experimental/UAF/UAFPoseSearch/Source/UAFPoseSearch/Private/HistoryCollectorTrait.cpp b/Engine/Plugins/Experimental/UAF/UAFPoseSearch/Source/UAFPoseSearch/Private/HistoryCollectorTrait.cpp
--- a/Engine/Plugins/Experimental/UAF/UAFPoseSearch/Source/UAFPoseSearch/Private/HistoryCollectorTrait.cpp
+++ b/Engine/Plugins/Experimental/UAF/UAFPoseSearch/Source/UAFPoseSearch/Private/HistoryCollectorTrait.cpp
@@ -146,10 +146,10 @@ void FAnimNextHistoryCollectorTask::Execute(UE::UAF::FEvaluationVM& VM) const
 			const FLODPoseStack& Pose = (*Keyframe)->Pose;
 			TArray<FBoneIndexType> BoneIndicesWithParents;
 
-			if (SharedData->PoseCount != PoseHistory->GetMaxNumPoses() || SharedData->SamplingInterval != PoseHistory->GetSamplingInterval())
+			if (PoseCount != PoseHistory->GetMaxNumPoses() || SamplingInterval != PoseHistory->GetSamplingInterval())
 			{
 				bCacheBones = true;
-				PoseHistory->Initialize_AnyThread(SharedData->PoseCount, SharedData->SamplingInterval);
+				PoseHistory->Initialize_AnyThread(PoseCount, SamplingInterval);
 			}
 
 			if (bCacheBones)
@@ -159,7 +159,7 @@ void FAnimNextHistoryCollectorTask::Execute(UE::UAF::FEvaluationVM& VM) const
 
 				BoneIndicesWithParents.Add(UE::PoseSearch::RootBoneIndexType);
 
-				for (const FBoneReference& CollectedBone : SharedData->CollectedBones)
+				for (const FBoneReference& CollectedBone : CollectedBones)
 				{
 					if (CollectedBone.BoneName != NAME_None)
 					{
@@ -179,10 +179,10 @@ void FAnimNextHistoryCollectorTask::Execute(UE::UAF::FEvaluationVM& VM) const
 
 			FComponentSpacePoseProvider ComponentSpacePoseProvider(Pose);
 			PoseHistory->EvaluateComponentSpace_AnyThread(DeltaTime, ComponentSpacePoseProvider, bStoreScales,
-				SharedData->RootBoneRecoveryTime, SharedData->RootBoneTranslationRecoveryRatio, SharedData->RootBoneRotationRecoveryRatio,
+				RootBoneRecoveryTime, RootBoneTranslationRecoveryRatio, RootBoneRotationRecoveryRatio,
 				bNeedsReset, bCacheBones, BoneIndicesWithParents
+				, FBlendedCurve(), CollectedCurves
 #if WITH_EDITORONLY_DATA
-				, FBlendedCurve(), TConstArrayView<FName>()
 				, AnimContext.IsValid() ? UE::PoseSearch::GetContextTransform(AnimContext.Pin().Get(), false) : FTransform::Identity
 #endif //WITH_EDITORONLY_DATA
 			);
@@ -225,23 +225,31 @@ namespace UE::UAF
 		// @todo: implement bResetOnBecomingRelevant
 		//bool bResetOnBecomingRelevant = SharedData->GetbResetOnBecomingRelevant(Binding);
 
-		if (SharedData->GetbGenerateTrajectory(Binding))
+		check(InstanceData->PoseHistoryPtr.IsValid());
+		InstanceData->PoseHistoryPtr->bGenerateTrajectory = SharedData->GetbGenerateTrajectory(Binding);
+		if (InstanceData->PoseHistoryPtr->bGenerateTrajectory)
 		{
-			// @todo: WIP trajectory generation
-			//FPoseSearchTrajectoryData::FSampling TrajectoryDataSampling;
-			//TrajectoryDataSampling.NumHistorySamples = FMath::Max(SharedData->PoseCount, SharedData->TrajectoryHistoryCount);
-			//TrajectoryDataSampling.SecondsPerHistorySample = SharedData->SamplingInterval;
-			//TrajectoryDataSampling.NumPredictionSamples = SharedData->TrajectoryPredictionCount;
-			//TrajectoryDataSampling.SecondsPerPredictionSample = SharedData->PredictionSamplingInterval;
-			//const bool bInitializeWithRefPose = SharedData->GetbInitializeWithRefPose(Binding);
-			//InstanceData->PoseHistory.GenerateTrajectory(InAnimInstance, InAnimInstance->GetDeltaSeconds(), InstanceData->TrajectoryData, TrajectoryDataSampling);
+			const UObject* AnimContext = nullptr;
+			if (const FAnimNextModuleInstance* ModuleInstance = Context.GetRootGraphInstance().GetModuleInstance())
+			{
+				AnimContext = ModuleInstance->GetObject();
+			}
+
+			InstanceData->PoseHistoryPtr->TrajectoryDataSampling.NumHistorySamples = FMath::Max(SharedData->GetPoseCount(Binding), SharedData->GetTrajectoryHistoryCount(Binding));
+			InstanceData->PoseHistoryPtr->TrajectoryDataSampling.SecondsPerHistorySample = SharedData->GetSamplingInterval(Binding);
+			InstanceData->PoseHistoryPtr->TrajectoryDataSampling.NumPredictionSamples = SharedData->GetTrajectoryPredictionCount(Binding);
+			InstanceData->PoseHistoryPtr->TrajectoryDataSampling.SecondsPerPredictionSample = SharedData->GetPredictionSamplingInterval(Binding);
+			InstanceData->PoseHistoryPtr->TrajectoryData = SharedData->GetTrajectoryData(Binding);
+
+			InstanceData->PoseHistoryPtr->GenerateTrajectory(AnimContext, InstanceData->DeltaTime);
 		}
 		else
 		{
-			check(InstanceData->PoseHistoryPtr.IsValid());
-			InstanceData->PoseHistoryPtr->SetTrajectory(SharedData->GetTrajectory(Binding));
+			InstanceData->PoseHistoryPtr->SetTrajectory(SharedData->GetTrajectory(Binding), SharedData->GetTrajectorySpeedMultiplier(Binding));
 		}
 
+		InstanceData->PoseHistoryPtr->bIsTrajectoryGeneratedBeforePreUpdate = false;
+
 		Context.PushScopedInterface<IPoseHistory>(Binding);
 
 		IUpdate::PreUpdate(Context, Binding, TraitState);
@@ -270,7 +278,6 @@ namespace UE::UAF
 
 		const FSharedData* SharedData = Binding.GetSharedData<FSharedData>();
 		FInstanceData* InstanceData = Binding.GetInstanceData<FInstanceData>();
-		const bool bStoreScales = SharedData->GetbStoreScales(Binding);
 
 #if WITH_EDITOR
 		if (InstanceData->bIsPostEvaluateBeingCalled)
@@ -296,10 +303,16 @@ namespace UE::UAF
 		check(InstanceData->PoseHistoryPtr.IsValid());
 		FAnimNextHistoryCollectorTask Task;
 		Task.PoseHistory = InstanceData->PoseHistoryPtr.Get();
-		Task.SharedData = SharedData;
-		Task.bStoreScales = bStoreScales;
+		Task.bStoreScales = SharedData->GetbStoreScales(Binding);
 		Task.DeltaTime = InstanceData->DeltaTime;
 		Task.HostObject = Context.GetHostObject();
+		Task.PoseCount = SharedData->GetPoseCount(Binding);
+		Task.SamplingInterval = SharedData->GetSamplingInterval(Binding);
+		Task.RootBoneRecoveryTime = SharedData->GetRootBoneRecoveryTime(Binding);
+		Task.RootBoneTranslationRecoveryRatio = SharedData->GetRootBoneTranslationRecoveryRatio(Binding);
+		Task.RootBoneRotationRecoveryRatio = SharedData->GetRootBoneRotationRecoveryRatio(Binding);
+		Task.CollectedBones = SharedData->GetCollectedBones(Binding);
+		Task.CollectedCurves = SharedData->GetCollectedCurves(Binding);
 
 #if WITH_EDITORONLY_DATA
 		if (const FAnimNextModuleInstance* ModuleInstance = Context.GetRootGraphInstance().GetModuleInstance())
diff --git a/Engine/Plugins/Experimental/UAF/UAFPoseSearch/Source/UAFPoseSearch/Private/HistoryCollectorTrait.h b/Engine/Plugins/Experimental/UAF/UAFPoseSearch/Source/UAFPoseSearch/Private/HistoryCollectorTrait.h
--- a/Engine/Plugins/Experimental/UAF/UAFPoseSearch/Source/UAFPoseSearch/Private/HistoryCollectorTrait.h
+++ b/Engine/Plugins/Experimental/UAF/UAFPoseSearch/Source/UAFPoseSearch/Private/HistoryCollectorTrait.h
@@ -23,13 +23,22 @@ struct FAnimNextHistoryCollectorTask : public FAnimNextEvaluationTask
 	virtual void Execute(UE::UAF::FEvaluationVM& VM) const override;
 
 	UE::PoseSearch::FPoseHistory* PoseHistory = nullptr;
-	const FAnimNextHistoryCollectorTraitSharedData* SharedData = nullptr;
 	bool bStoreScales = false;
 	float DeltaTime = 0.f;
 	TWeakObjectPtr<const UObject> HostObject;
 #if WITH_EDITORONLY_DATA
 	TWeakObjectPtr<const UObject> AnimContext;
 #endif // WITH_EDITORONLY_DATA
+
+	int32 PoseCount = 0;
+	float SamplingInterval = 0.f;
+	float RootBoneRecoveryTime = 0.f;
+	float RootBoneTranslationRecoveryRatio = 0.f;
+	float RootBoneRotationRecoveryRatio = 0.f;
+
+	// @todo: minimize allocations
+	TArray<FBoneReference> CollectedBones;
+	TArray<FName> CollectedCurves;
 };
 
 namespace UE::UAF
