PoseSearch - consolidating FMotionMatchingTrait to FAnimNode_MotionMatching in handling empty or filled with null databases to search requests

#Synced-CL 50943617

diff --git a/Engine/Plugins/Animation/PoseSearch/Source/Runtime/Private/PoseSearchContext.cpp b/Engine/Plugins/Animation/PoseSearch/Source/Runtime/Private/PoseSearchContext.cpp
--- a/Engine/Plugins/Animation/PoseSearch/Source/Runtime/Private/PoseSearchContext.cpp
+++ b/Engine/Plugins/Animation/PoseSearch/Source/Runtime/Private/PoseSearchContext.cpp
@@ -1339,7 +1339,7 @@ bool FSearchContext::CanUseContinuingPoseValues() const
 #if UE_POSE_SEARCH_TRACE_ENABLED
 void FSearchContext::Track(const UPoseSearchDatabase* Database, int32 PoseIdx, EPoseCandidateFlags PoseCandidateFlags, const FPoseSearchCost& Cost)
 {
-	if (ensure(Database))
+	if (Database)
 	{
 		FBestPoseCandidates& BestPoseCandidates = BestPoseCandidatesMap.FindOrAdd(Database);
 		if (PoseIdx != INDEX_NONE)
diff --git a/Engine/Plugins/Experimental/UAF/UAFPoseSearch/Source/UAFPoseSearch/Private/MotionMatchingTrait.cpp b/Engine/Plugins/Experimental/UAF/UAFPoseSearch/Source/UAFPoseSearch/Private/MotionMatchingTrait.cpp
--- a/Engine/Plugins/Experimental/UAF/UAFPoseSearch/Source/UAFPoseSearch/Private/MotionMatchingTrait.cpp
+++ b/Engine/Plugins/Experimental/UAF/UAFPoseSearch/Source/UAFPoseSearch/Private/MotionMatchingTrait.cpp
@@ -199,86 +199,76 @@ PRAGMA_ENABLE_DEPRECATION_WARNINGS
 			//MotionMatchingState.SearchResult.BlendParameters = GetBlendParameters();
 
 			const TConstArrayView<TObjectPtr<const UPoseSearchDatabase>> DatabasesToSearch = SharedData->GetDatabases(Binding);
-			if (DatabasesToSearch.IsEmpty())
-			{
-				// if we have availabilities, it's ok to have a FMotionMatchingTrait set up purely for interactions, without any other database to search! 
-				if (Availabilities.IsEmpty())
-				{
-					UE_LOG(LogPoseSearch, Error, TEXT("FMotionMatchingTrait::PreUpdate, No database assets provided for motion matching."));
-				}
-			}
-			else
-			{
-				FFloatInterval PoseSearchPlayRate = SharedData->GetPlayRate(Binding);
+			
+			FFloatInterval PoseSearchPlayRate = SharedData->GetPlayRate(Binding);
 #if ENABLE_ANIM_DEBUG
-				bool bPoseSearchPlayRateEnabled = GVarAnimNextMotionMatchingPlayRateEnabled == UE::Private::EPlayRateState::Enabled || GVarAnimNextMotionMatchingPlayRateEnabled == UE::Private::EPlayRateState::PoseSearchOnly;
-				PoseSearchPlayRate = bPoseSearchPlayRateEnabled ? PoseSearchPlayRate : FFloatInterval(1.f, 1.f);
+			bool bPoseSearchPlayRateEnabled = GVarAnimNextMotionMatchingPlayRateEnabled == UE::Private::EPlayRateState::Enabled || GVarAnimNextMotionMatchingPlayRateEnabled == UE::Private::EPlayRateState::PoseSearchOnly;
+			PoseSearchPlayRate = bPoseSearchPlayRateEnabled ? PoseSearchPlayRate : FFloatInterval(1.f, 1.f);
 #endif // ENABLE_ANIM_DEBUG
 
-				const float SearchThrottleTime = SharedData->GetbShouldSearch(Binding) ? SharedData->GetSearchThrottleTime(Binding) : UE_BIG_NUMBER;
+			const float SearchThrottleTime = SharedData->GetbShouldSearch(Binding) ? SharedData->GetSearchThrottleTime(Binding) : UE_BIG_NUMBER;
 
-				FPoseSearchEvent EventToSearch = SharedData->GetEventToSearch(Binding);
-				FSearchContext SearchContext(0.f, SharedData->GetPoseJumpThresholdTime(Binding), EventToSearch.GetPlayRateOverriddenEvent(PoseSearchPlayRate));
+			FPoseSearchEvent EventToSearch = SharedData->GetEventToSearch(Binding);
+			FSearchContext SearchContext(0.f, SharedData->GetPoseJumpThresholdTime(Binding), EventToSearch.GetPlayRateOverriddenEvent(PoseSearchPlayRate));
 
-				FChooserEvaluationContext ChooserEvaluationContext(const_cast<UObject*>(AnimContext));
-				const FUAFAssetInstance& Instance = Binding.GetTraitPtr().GetNodeInstance()->GetOwner();
-				ChooserEvaluationContext.AddStructViewParam(FStructView::Make(const_cast<FUAFAssetInstance&>(Instance)));
-				SearchContext.AddRole(GetCommonDefaultRole(DatabasesToSearch), &ChooserEvaluationContext, PoseHistory);
+			FChooserEvaluationContext ChooserEvaluationContext(const_cast<UObject*>(AnimContext));
+			const FUAFAssetInstance& Instance = Binding.GetTraitPtr().GetNodeInstance()->GetOwner();
+			ChooserEvaluationContext.AddStructViewParam(FStructView::Make(const_cast<FUAFAssetInstance&>(Instance)));
+			SearchContext.AddRole(GetCommonDefaultRole(DatabasesToSearch), &ChooserEvaluationContext, PoseHistory);
 
-				// If we can't advance or enough time has elapsed since the last pose jump then search
-				bool bSearch = true;
-				FSearchResult SearchResult;
-				if (MotionMatchingState.ElapsedPoseSearchTime < SearchThrottleTime)
-				{
-					SearchResult.InitFrom(MotionMatchingState.SearchResult);
-					SearchContext.UpdateContinuingPoseSearchResult(SearchResult, SearchResult);
-					const bool bCanAdvance = SearchContext.GetContinuingPoseSearchResult().PoseIdx != INDEX_NONE;
-					bSearch = !bCanAdvance;
-				}
+			// If we can't advance or enough time has elapsed since the last pose jump then search
+			bool bSearch = true;
+			FSearchResult SearchResult;
+			if (MotionMatchingState.ElapsedPoseSearchTime < SearchThrottleTime)
+			{
+				SearchResult.InitFrom(MotionMatchingState.SearchResult);
+				SearchContext.UpdateContinuingPoseSearchResult(SearchResult, SearchResult);
+				const bool bCanAdvance = SearchContext.GetContinuingPoseSearchResult().PoseIdx != INDEX_NONE;
+				bSearch = !bCanAdvance;
+			}
 
-				if (bSearch)
-				{
-					MotionMatchingState.ElapsedPoseSearchTime = 0.f;
+			if (bSearch)
+			{
+				MotionMatchingState.ElapsedPoseSearchTime = 0.f;
 			
-					TArrayView<const UObject*> AssetsToSearch((const UObject**)FMemory_Alloca(DatabasesToSearch.Num() * sizeof(UObject*)), (int32)DatabasesToSearch.Num());
-					for (int32 DatabaseToSearchIndex = 0; DatabaseToSearchIndex < DatabasesToSearch.Num(); ++DatabaseToSearchIndex)
-					{
-						AssetsToSearch[DatabaseToSearchIndex] = DatabasesToSearch[DatabaseToSearchIndex].Get();
-					}
-
-					FPoseSearchContinuingProperties ContinuingProperties;
-					ContinuingProperties.InitFrom(MotionMatchingState.SearchResult, SharedData->GetNextUpdateInterruptMode(Binding));
-					FSearchResults_Single SearchResults;
-					UPoseSearchLibrary::MotionMatch(SearchContext, AssetsToSearch, ContinuingProperties, SearchResults);
-					SearchResult = SearchResults.GetBestResult();
-				}
-				else
+				TArrayView<const UObject*> AssetsToSearch((const UObject**)FMemory_Alloca(DatabasesToSearch.Num() * sizeof(UObject*)), (int32)DatabasesToSearch.Num());
+				for (int32 DatabaseToSearchIndex = 0; DatabaseToSearchIndex < DatabasesToSearch.Num(); ++DatabaseToSearchIndex)
 				{
-					MotionMatchingState.ElapsedPoseSearchTime += DeltaTime;
+					AssetsToSearch[DatabaseToSearchIndex] = DatabasesToSearch[DatabaseToSearchIndex].Get();
+				}
+
+				FPoseSearchContinuingProperties ContinuingProperties;
+				ContinuingProperties.InitFrom(MotionMatchingState.SearchResult, SharedData->GetNextUpdateInterruptMode(Binding));
+				FSearchResults_Single SearchResults;
+				UPoseSearchLibrary::MotionMatch(SearchContext, AssetsToSearch, ContinuingProperties, SearchResults);
+				SearchResult = SearchResults.GetBestResult();
+			}
+			else
+			{
+				MotionMatchingState.ElapsedPoseSearchTime += DeltaTime;
 
-					// if we didn't perform a search we carry on with the continuing pose search result
-					SearchResult = SearchContext.GetContinuingPoseSearchResult();
-					SearchResult.bIsContinuingPoseSearch = true;
+				// if we didn't perform a search we carry on with the continuing pose search result
+				SearchResult = SearchContext.GetContinuingPoseSearchResult();
+				SearchResult.bIsContinuingPoseSearch = true;
 
 #if UE_POSE_SEARCH_TRACE_ENABLED
-					// in case we skipped the search, we still have to track we would have requested to evaluate DatabasesToSearch and SearchResult.Database
-					for (const TObjectPtr<const UPoseSearchDatabase>& DatabaseToSearch : DatabasesToSearch)
-					{
-						SearchContext.Track(DatabaseToSearch.Get());
-					}
+				// in case we skipped the search, we still have to track we would have requested to evaluate DatabasesToSearch and SearchResult.Database
+				for (const TObjectPtr<const UPoseSearchDatabase>& DatabaseToSearch : DatabasesToSearch)
+				{
+					SearchContext.Track(DatabaseToSearch.Get());
+				}
 
-					SearchContext.Track(SearchResult.Database.Get());
+				SearchContext.Track(SearchResult.Database.Get());
 #endif // UE_POSE_SEARCH_TRACE_ENABLED
-				}
+			}
 
-				const float WantedPlayRate = CalculateWantedPlayRate(SearchResult, SearchContext, PoseSearchPlayRate, PoseHistory ? PoseHistory->GetTrajectorySpeedMultiplier() : 1.f, EventToSearch);
-				MotionMatchingState.SearchResult.InitFrom(SearchResult, WantedPlayRate);
+			const float WantedPlayRate = CalculateWantedPlayRate(SearchResult, SearchContext, PoseSearchPlayRate, PoseHistory ? PoseHistory->GetTrajectorySpeedMultiplier() : 1.f, EventToSearch);
+			MotionMatchingState.SearchResult.InitFrom(SearchResult, WantedPlayRate);
 
-				if (const FPoseIndicesHistory* PoseIndicesHistory = PoseHistory->GetPoseIndicesHistory())
-				{
-					// const casting here is safe since we're in the thread owning the pose history, and it's the correct place to update the previously selected poses
-					const_cast<FPoseIndicesHistory*>(PoseIndicesHistory)->Update(SearchResult, DeltaTime, SharedData->GetPoseReselectHistory(Binding));
-				}
+			if (const FPoseIndicesHistory* PoseIndicesHistory = PoseHistory->GetPoseIndicesHistory())
+			{
+				// const casting here is safe since we're in the thread owning the pose history, and it's the correct place to update the previously selected poses
+				const_cast<FPoseIndicesHistory*>(PoseIndicesHistory)->Update(SearchResult, DeltaTime, SharedData->GetPoseReselectHistory(Binding));
 			}
 		}
 
