AnimNode Starting Point - Pose Search

#Synced-CL 51055096

diff --git a/Engine/Plugins/Experimental/UAF/UAFAnimNode/Source/UAFAnimNode/Public/UAF/AnimNodeCore/UAFAnimNodeData.h b/Engine/Plugins/Experimental/UAF/UAFAnimNode/Source/UAFAnimNode/Public/UAF/AnimNodeCore/UAFAnimNodeData.h
--- a/Engine/Plugins/Experimental/UAF/UAFAnimNode/Source/UAFAnimNode/Public/UAF/AnimNodeCore/UAFAnimNodeData.h
+++ b/Engine/Plugins/Experimental/UAF/UAFAnimNode/Source/UAFAnimNode/Public/UAF/AnimNodeCore/UAFAnimNodeData.h
@@ -34,7 +34,6 @@ namespace UE::UAF
 		template<typename InterfaceType> InterfaceType* GetInterface();
 		template<typename InterfaceType> const InterfaceType* GetInterface() const;
 
-	protected:
 		virtual void* GetInterface(FUAFAnimNodeDataInterfaceId Id);
 	};
 
diff --git a/Engine/Plugins/Experimental/UAF/UAFPoseSearch/Source/UAFPoseSearch/Private/AnimNode/UAFPoseHistoryNode.cpp b/Engine/Plugins/Experimental/UAF/UAFPoseSearch/Source/UAFPoseSearch/Private/AnimNode/UAFPoseHistoryNode.cpp
new file mode 100644
--- /dev/null
+++ b/Engine/Plugins/Experimental/UAF/UAFPoseSearch/Source/UAFPoseSearch/Private/AnimNode/UAFPoseHistoryNode.cpp
@@ -0,0 +1,116 @@
+// Copyright Epic Games, Inc. All Rights Reserved.
+
+#include "UAFPoseHistoryNode.h"
+#include "Animation/AnimNodeBase.h"
+#include "GenerationTools.h"
+#include "PoseHistoryEvaluation.h"
+#include "EvaluationVM/Tasks/PushReferenceKeyframe.h"
+#include "UAF/AnimNodeCore/UAFAnimNodeUpdate.h"
+
+namespace UE::UAF
+{
+	
+FUAFPoseHistoryNode::FUAFPoseHistoryNode(const FUAFPoseHistoryNodeData* InData, UE::UAF::FUAFAnimGraphUpdateContext& Context)
+	: Data(InData)
+{
+	if (InData->ChildNode.IsValid())
+	{
+		ChildNode = InData->ChildNode.Get().CreateInstance(Context);
+	}
+}
+
+void FUAFPoseHistoryNode::Update(FUAFAnimGraphUpdateContext& Context, float DeltaTime)
+{
+	SCOPED_NAMED_EVENT(AnimNode_Update_PoseHistory, FColor::Blue);
+	if (Data==nullptr)
+	{
+		// this should never happen... but fallback TPose in case
+		Context.AppendTask(FAnimNextPushReferenceKeyframeTask::MakeFromSkeleton());
+		return;
+	}
+
+	InstanceData.DeltaTime = DeltaTime;
+
+	check(InstanceData.PoseHistoryPtr.IsValid());
+
+	if (Data->Trajectory.IsValid())
+	{
+		FTransformTrajectory TransformTrajectory;
+		// probably need a way to get a pointer to trajectory to avoid array copies
+		Context.GetVariablesOwner()->GetVariable(Data->Trajectory, TransformTrajectory);
+		InstanceData.PoseHistoryPtr->SetTrajectory(TransformTrajectory);
+	}
+
+	// pre-evaluate 
+	FAnimNextHistoryCollectorPreEvaluateTask PreEvaluateTask;
+	PreEvaluateTask.InstanceData = &InstanceData;
+	Context.AppendTask(PreEvaluateTask);
+
+	if (Data->PoseHistory.IsValid())
+	{
+		FPoseHistoryReference Reference;
+		Reference.PoseHistory = InstanceData.PoseHistoryPtr;
+		Context.GetVariablesOwner()->SetVariable(Data->PoseHistory, Reference);
+
+	}
+
+	if (Data->Settings.PoseCount != InstanceData.PoseHistoryPtr->GetMaxNumPoses() || Data->Settings.SamplingInterval != InstanceData.PoseHistoryPtr->GetSamplingInterval())
+	{
+		InstanceData.PoseHistoryPtr->Initialize_AnyThread(Data->Settings.PoseCount, Data->Settings.SamplingInterval);
+	}
+
+	if (ChildNode)
+	{
+		// update + evaluate child
+		ChildNode->Update(Context, DeltaTime);
+	}
+	else
+	{
+		// fallback T-Pose if child node is not valid
+		Context.AppendTask(FAnimNextPushReferenceKeyframeTask::MakeFromSkeleton());
+	}
+
+	// post-evaluate
+	FAnimNextHistoryCollectorTask PostEvaluateTask;
+	PostEvaluateTask.PoseHistory = InstanceData.PoseHistoryPtr.Get();
+	PostEvaluateTask.CollectedBones = Data->Settings.CollectedBones;
+	PostEvaluateTask.CollectedCurves = Data->Settings.CollectedCurves;
+	PostEvaluateTask.bStoreScales = Data->Settings.bStoreScales;
+	PostEvaluateTask.PoseCount = Data->Settings.PoseCount;
+	PostEvaluateTask.SamplingInterval = Data->Settings.SamplingInterval;
+	PostEvaluateTask.DeltaTime = DeltaTime;
+	PostEvaluateTask.HostObject = Context.GetHostObject();
+
+#if WITH_EDITORONLY_DATA
+	PostEvaluateTask.AnimContext = Context.GetHostObject();
+#endif // WITH_EDITORONLY_DATA
+	
+	Context.AppendTask(PostEvaluateTask);	
+}
+
+void* FUAFPoseHistoryNode::GetInterface(UE::UAF::FUAFAnimNodeInterfaceId Id)
+{
+	if (ChildNode)
+	{
+		return ChildNode->GetInterface(Id);
+	}
+
+	return nullptr;
+}
+
+TSharedPtr<FUAFAnimNode> FUAFPoseHistoryNodeData::CreateInstance(FUAFAnimGraphUpdateContext& Context) const
+{
+	return MakeShared<FUAFPoseHistoryNode>(this, Context);
+}
+
+void* FUAFPoseHistoryNodeData::GetInterface(FUAFAnimNodeInterfaceId Id)
+{
+	if (ChildNode.IsValid())
+	{
+		return ChildNode.GetMutable().GetInterface(Id);
+	}
+
+	return nullptr;
+}
+	
+}
diff --git a/Engine/Plugins/Experimental/UAF/UAFPoseSearch/Source/UAFPoseSearch/Private/AnimNode/UAFPoseHistoryNode.h b/Engine/Plugins/Experimental/UAF/UAFPoseSearch/Source/UAFPoseSearch/Private/AnimNode/UAFPoseHistoryNode.h
new file mode 100644
--- /dev/null
+++ b/Engine/Plugins/Experimental/UAF/UAFPoseSearch/Source/UAFPoseSearch/Private/AnimNode/UAFPoseHistoryNode.h
@@ -0,0 +1,59 @@
+// Copyright Epic Games, Inc. All Rights Reserved.
+
+#pragma once
+
+#include "CoreMinimal.h"
+#include "TraitInterfaces/IUpdate.h"
+#include "UAF/AnimNodeCore/UAFAnimNode.h"
+#include "HistoryCollectorTraitData.h"
+#include "../HistoryCollectorTrait.h"
+#include "UAF/AnimNodeCore/UAFAnimNodeData.h"
+#include "UAFPoseHistoryNode.generated.h"
+
+namespace UE::UAF
+{
+	
+struct FUAFPoseHistoryNodeData;
+	
+struct FUAFPoseHistoryNode : public FUAFAnimNode
+{
+	FUAFPoseHistoryNode() = default;
+	FUAFPoseHistoryNode(const FUAFPoseHistoryNodeData* InData, FUAFAnimGraphUpdateContext& Context);
+
+	virtual void Update(FUAFAnimGraphUpdateContext& Context, float DeltaTime) override;
+
+	virtual void* GetInterface(FUAFAnimNodeInterfaceId Id) override;
+
+	TSharedPtr<FUAFAnimNode> ChildNode;
+
+	FHistoryCollectorTrait::FInstanceData InstanceData;
+
+	const FUAFPoseHistoryNodeData* Data = nullptr;
+};
+
+
+USTRUCT(DisplayName = "Pose History")
+struct FUAFPoseHistoryNodeData : public FUAFAnimNodeData
+{
+	GENERATED_BODY()
+
+	UPROPERTY(EditAnywhere, Category = "Data",  meta = (ExcludeBaseStruct))
+	TInstancedStruct<FUAFAnimNodeData> ChildNode;
+
+	// todo: this should be some sort of output binding when we build that system
+	UPROPERTY(EditAnywhere, Category = "NodeData")
+	FAnimNextVariableReference PoseHistory;
+
+	// todo: this should be some sort of input binding when we build that system
+	UPROPERTY(EditAnywhere, Category = "NodeData")
+	FAnimNextVariableReference Trajectory;
+
+	// Use Trait Shared Data for settings to keep feature set in sync
+	UPROPERTY(EditAnywhere, Category = "NodeData")
+	FAnimNextHistoryCollectorTraitSharedData Settings;
+
+	virtual TSharedPtr<FUAFAnimNode> CreateInstance(FUAFAnimGraphUpdateContext& Context) const override;
+	virtual void* GetInterface(FUAFAnimNodeInterfaceId Id) override;
+};
+	
+}
\ No newline at end of file
diff --git a/Engine/Plugins/Experimental/UAF/UAFPoseSearch/Source/UAFPoseSearch/UAFPoseSearch.Build.cs b/Engine/Plugins/Experimental/UAF/UAFPoseSearch/Source/UAFPoseSearch/UAFPoseSearch.Build.cs
--- a/Engine/Plugins/Experimental/UAF/UAFPoseSearch/Source/UAFPoseSearch/UAFPoseSearch.Build.cs
+++ b/Engine/Plugins/Experimental/UAF/UAFPoseSearch/Source/UAFPoseSearch/UAFPoseSearch.Build.cs
@@ -6,6 +6,8 @@ namespace UnrealBuildTool.Rules
 	{
 		public UAFPoseSearch(ReadOnlyTargetRules Target) : base(Target)
 		{
+			bAllowUETypesInNamespaces = true;
+
 			PublicDependencyModuleNames.AddRange(
 				new string[]
 				{
@@ -19,7 +21,8 @@ namespace UnrealBuildTool.Rules
 					"PoseSearch",
 					"RigVM",
 					"UAF",
-					"UAFAnimGraph"
+					"UAFAnimGraph",
+					"UAFAnimNode"
 				}
 			);
 		}
diff --git a/Engine/Plugins/Experimental/UAF/UAFPoseSearch/UAFPoseSearch.uplugin b/Engine/Plugins/Experimental/UAF/UAFPoseSearch/UAFPoseSearch.uplugin
--- a/Engine/Plugins/Experimental/UAF/UAFPoseSearch/UAFPoseSearch.uplugin
+++ b/Engine/Plugins/Experimental/UAF/UAFPoseSearch/UAFPoseSearch.uplugin
@@ -36,7 +36,11 @@
 		{
 			"Name": "UAFAnimGraph",
 			"Enabled": true
-		}		
+		},
+		{
+			"Name": "UAFAnimNode",
+			"Enabled": true
+		}
 	],
 	"Modules" :
 	[
