UAF: Fix for curves not working in Pose History trait
#rnx
#Synced-CL 51342373

diff --git a/Engine/Plugins/Experimental/UAF/UAFPoseSearch/Source/UAFPoseSearch/Private/HistoryCollectorTrait.cpp b/Engine/Plugins/Experimental/UAF/UAFPoseSearch/Source/UAFPoseSearch/Private/HistoryCollectorTrait.cpp
--- a/Engine/Plugins/Experimental/UAF/UAFPoseSearch/Source/UAFPoseSearch/Private/HistoryCollectorTrait.cpp
+++ b/Engine/Plugins/Experimental/UAF/UAFPoseSearch/Source/UAFPoseSearch/Private/HistoryCollectorTrait.cpp
@@ -177,11 +177,13 @@ void FAnimNextHistoryCollectorTask::Execute(UE::UAF::FEvaluationVM& VM) const
 				FAnimationRuntime::EnsureParentsPresent(BoneIndicesWithParents, Skeleton->GetReferenceSkeleton());
 			}
 
+			const bool bHasCurves = EnumHasAnyFlags(VM.GetFlags(), EEvaluationFlags::Curves);
+
 			FComponentSpacePoseProvider ComponentSpacePoseProvider(Pose);
 			PoseHistory->EvaluateComponentSpace_AnyThread(DeltaTime, ComponentSpacePoseProvider, bStoreScales,
 				RootBoneRecoveryTime, RootBoneTranslationRecoveryRatio, RootBoneRotationRecoveryRatio,
 				bNeedsReset, bCacheBones, BoneIndicesWithParents
-				, FBlendedCurve(), CollectedCurves
+				, bHasCurves ? (*Keyframe)->Curves : FBlendedCurve(), CollectedCurves
 #if WITH_EDITORONLY_DATA
 				, AnimContext.IsValid() ? UE::PoseSearch::GetContextTransform(AnimContext.Pin().Get(), false) : FTransform::Identity
 #endif //WITH_EDITORONLY_DATA
