UAFPoseSearch - consolidating FUAFDebuggerTrackCreator and FDebuggerTrackCreator, since GetTargetTypeNameInternal now returns UObject::StaticClass()->GetFName(), to avoid having two different pose search debugger tracks with UAF

#Synced-CL 52360741

diff --git a/Engine/Plugins/Experimental/UAF/UAFPoseSearch/Source/UAFPoseSearchUncookedOnly/Private/UAFPoseSearchDebugger.cpp b/Engine/Plugins/Experimental/UAF/UAFPoseSearch/Source/UAFPoseSearchUncookedOnly/Private/UAFPoseSearchDebugger.cpp
deleted file mode 100644
diff --git a/Engine/Plugins/Experimental/UAF/UAFPoseSearch/Source/UAFPoseSearchUncookedOnly/Private/UAFPoseSearchDebugger.h b/Engine/Plugins/Experimental/UAF/UAFPoseSearch/Source/UAFPoseSearchUncookedOnly/Private/UAFPoseSearchDebugger.h
deleted file mode 100644
diff --git a/Engine/Plugins/Experimental/UAF/UAFPoseSearch/Source/UAFPoseSearchUncookedOnly/Private/UAFPoseSearchUncookedOnlyModule.cpp b/Engine/Plugins/Experimental/UAF/UAFPoseSearch/Source/UAFPoseSearchUncookedOnly/Private/UAFPoseSearchUncookedOnlyModule.cpp
--- a/Engine/Plugins/Experimental/UAF/UAFPoseSearch/Source/UAFPoseSearchUncookedOnly/Private/UAFPoseSearchUncookedOnlyModule.cpp
+++ b/Engine/Plugins/Experimental/UAF/UAFPoseSearch/Source/UAFPoseSearchUncookedOnly/Private/UAFPoseSearchUncookedOnlyModule.cpp
@@ -1,43 +1,17 @@
 ﻿// Copyright Epic Games, Inc. All Rights Reserved.
 
-#include "PoseSearchDebuggerTrackCreator.h"
-#include "UAFPoseSearchDebugger.h"
 #include "Features/IModularFeatures.h"
 #include "Modules/ModuleManager.h"
 
 namespace UE::PoseSearch
 {
-class FUAFDebuggerTrackCreator;
 
 //////////////////////////////////////////////////////////////////////////
 // FPoseSearchEditorModule
 class FUAFPoseSearchUncookedOnlyModule : public IModuleInterface
 {
-public:
-	void StartupModule() override;
-	void ShutdownModule() override;
-
-private:
-	
-	/** Creates the track for the Rewind Debugger */
-	TSharedPtr<FUAFDebuggerTrackCreator> UAFDebuggerTrackCreator;
 };
 
-void FUAFPoseSearchUncookedOnlyModule::StartupModule()
-{
-	IModuleInterface::StartupModule();
-
-	UAFDebuggerTrackCreator = MakeShared<FUAFDebuggerTrackCreator>();
-	IModularFeatures::Get().RegisterModularFeature(RewindDebugger::IRewindDebuggerTrackCreator::ModularFeatureName, UAFDebuggerTrackCreator.Get());
-}
-
-void FUAFPoseSearchUncookedOnlyModule::ShutdownModule()
-{
-	IModularFeatures::Get().UnregisterModularFeature(RewindDebugger::IRewindDebuggerTrackCreator::ModularFeatureName, UAFDebuggerTrackCreator.Get());
-	
-	IModuleInterface::ShutdownModule();
-}
-
 } // namespace UE::PoseSearch
 
 IMPLEMENT_MODULE(UE::PoseSearch::FUAFPoseSearchUncookedOnlyModule, UAFPoseSearchUncookedOnly)
