PoseSearch - removing unnecessary MotionMatchingInteraction node

#Synced-CL 52360741

diff --git a/Engine/Plugins/Animation/PoseSearch/Source/Editor/Private/AnimGraphNode_MotionMatchingInteraction.cpp b/Engine/Plugins/Animation/PoseSearch/Source/Editor/Private/AnimGraphNode_MotionMatchingInteraction.cpp
deleted file mode 100644
diff --git a/Engine/Plugins/Animation/PoseSearch/Source/Editor/Public/AnimGraphNode_MotionMatchingInteraction.h b/Engine/Plugins/Animation/PoseSearch/Source/Editor/Public/AnimGraphNode_MotionMatchingInteraction.h
deleted file mode 100644
diff --git a/Engine/Plugins/Animation/PoseSearch/Source/Runtime/Private/AnimNode_MotionMatchingInteraction.cpp b/Engine/Plugins/Animation/PoseSearch/Source/Runtime/Private/AnimNode_MotionMatchingInteraction.cpp
deleted file mode 100644
diff --git a/Engine/Plugins/Animation/PoseSearch/Source/Runtime/Private/MotionMatchingInteractionAnimNodeLibrary.cpp b/Engine/Plugins/Animation/PoseSearch/Source/Runtime/Private/MotionMatchingInteractionAnimNodeLibrary.cpp
deleted file mode 100644
diff --git a/Engine/Plugins/Animation/PoseSearch/Source/Runtime/Public/PoseSearch/AnimNode_MotionMatchingInteraction.h b/Engine/Plugins/Animation/PoseSearch/Source/Runtime/Public/PoseSearch/AnimNode_MotionMatchingInteraction.h
deleted file mode 100644
diff --git a/Engine/Plugins/Animation/PoseSearch/Source/Runtime/Public/PoseSearch/MotionMatchingInteractionAnimNodeLibrary.h b/Engine/Plugins/Animation/PoseSearch/Source/Runtime/Public/PoseSearch/MotionMatchingInteractionAnimNodeLibrary.h
deleted file mode 100644
diff --git a/Engine/Plugins/Experimental/UAF/UAFPoseSearch/Source/UAFPoseSearch/Private/MotionMatchingTrait.cpp b/Engine/Plugins/Experimental/UAF/UAFPoseSearch/Source/UAFPoseSearch/Private/MotionMatchingTrait.cpp
--- a/Engine/Plugins/Experimental/UAF/UAFPoseSearch/Source/UAFPoseSearch/Private/MotionMatchingTrait.cpp
+++ b/Engine/Plugins/Experimental/UAF/UAFPoseSearch/Source/UAFPoseSearch/Private/MotionMatchingTrait.cpp
@@ -155,9 +155,9 @@ PRAGMA_ENABLE_DEPRECATION_WARNINGS
 		// performing the regular single character motion matching search in case there's no MM interaction
 		if (!MotionMatchingState.SearchResult.bIsInteraction)
 		{
-			// We adjust the MotionMatchingInteractionState time to the current player node's asset time. This is done 
+			// We adjust the motion matching state asset time to the current player node's asset time. This is done 
 			// because the player node may have ticked more or less time than we expected due to variable dt or the 
-			// dynamic playback rate adjustment and as such the MotionMatchingInteractionState does not update by itself
+			// dynamic playback rate adjustment and as such the motion matching state does not update by itself
 			float TimelineNormalizedTime = 0.0f;
 			FTimelineState TimelineState;
 			if (TimelineTrait.GetState(Context, TimelineState))
